[Horde][Analytics][Server] JobStepRef telemetry emissions for skipped steps.

#Synced-CL 49094299

diff --git a/Engine/Source/Programs/Horde/Plugins/Build/HordeServer.Build/Jobs/IJobStepRefCollection.cs b/Engine/Source/Programs/Horde/Plugins/Build/HordeServer.Build/Jobs/IJobStepRefCollection.cs
--- a/Engine/Source/Programs/Horde/Plugins/Build/HordeServer.Build/Jobs/IJobStepRefCollection.cs
+++ b/Engine/Source/Programs/Horde/Plugins/Build/HordeServer.Build/Jobs/IJobStepRefCollection.cs
@@ -18,7 +18,7 @@ namespace HordeServer.Jobs
 	public interface IJobStepRefCollection
 	{
 		/// <summary>
-		/// Constructor
+		/// Inserts or replaces a JobStepRef.
 		/// </summary>
 		/// <param name="id">Unique id for the step being referenced</param>
 		/// <param name="jobName">Name of the job</param>
@@ -38,9 +38,10 @@ namespace HordeServer.Jobs
 		/// <param name="initTime">Time taken for the batch containing this step to initializer</param>
 		/// <param name="jobStartTimeUtc">Start time of the job</param>
 		/// <param name="startTimeUtc">Start time</param>
-		/// <param name="finishTimeUtc">Finish time for the step, if known</param>		
-		/// <param name="stepMetadata">Step meta data</param>		
-		Task<IJobStepRef> InsertOrReplaceAsync(JobStepRefId id, string jobName, string stepName, StreamId streamId, TemplateId templateId, CommitIdWithOrder commitId, LogId? logId, PoolId? poolId, AgentId? agentId, JobStepState? state, JobStepOutcome? outcome, bool updateIssues, CommitIdWithOrder? lastSuccess, CommitIdWithOrder? lastWarning, float waitTime, float initTime, DateTime jobStartTimeUtc, DateTime startTimeUtc, DateTime? finishTimeUtc, List<string>? stepMetadata = null);
+		/// <param name="finishTimeUtc">Finish time for the step, if known</param>
+		/// <param name="stepMetadata">Step meta data</param>
+		/// <param name="logger">Logger to use throughout.</param>
+		Task<IJobStepRef> InsertOrReplaceAsync(JobStepRefId id, string jobName, string stepName, StreamId streamId, TemplateId templateId, CommitIdWithOrder commitId, LogId? logId, PoolId? poolId, AgentId? agentId, JobStepState? state, JobStepOutcome? outcome, bool updateIssues, CommitIdWithOrder? lastSuccess, CommitIdWithOrder? lastWarning, float waitTime, float initTime, DateTime jobStartTimeUtc, DateTime startTimeUtc, DateTime? finishTimeUtc, List<string>? stepMetadata = null, ILogger? logger = null);
 
 		/// <summary>
 		/// Updates a job step ref 
@@ -162,7 +163,7 @@ namespace HordeServer.Jobs
 					logger.LogInformation("Updating step reference {StepId} for job {JobId}, batch {BatchId}, with outcome {JobStepOutcome}", step.Id, job.Id, batch.Id, outcome);
 				}
 
-				await jobStepRefs.InsertOrReplaceAsync(new JobStepRefId(job.Id, batch.Id, step.Id), job.Name, nodeName, job.StreamId, job.TemplateId, job.CommitId, step.LogId, batch.PoolId, batch.AgentId, step.State, outcome, job.UpdateIssues, lastSuccess, lastWarning, waitTime, initTime, job.CreateTimeUtc, step.StartTimeUtc ?? DateTime.UtcNow, step.FinishTimeUtc, step.Metadata?.ToList());
+				await jobStepRefs.InsertOrReplaceAsync(new JobStepRefId(job.Id, batch.Id, step.Id), job.Name, nodeName, job.StreamId, job.TemplateId, job.CommitId, step.LogId, batch.PoolId, batch.AgentId, step.State, outcome, job.UpdateIssues, lastSuccess, lastWarning, waitTime, initTime, job.CreateTimeUtc, step.StartTimeUtc ?? DateTime.UtcNow, step.FinishTimeUtc, step.Metadata?.ToList(), logger);
 			}
 		}
 	}
diff --git a/Engine/Source/Programs/Horde/Plugins/Build/HordeServer.Build/Jobs/JobCollection.cs b/Engine/Source/Programs/Horde/Plugins/Build/HordeServer.Build/Jobs/JobCollection.cs
--- a/Engine/Source/Programs/Horde/Plugins/Build/HordeServer.Build/Jobs/JobCollection.cs
+++ b/Engine/Source/Programs/Horde/Plugins/Build/HordeServer.Build/Jobs/JobCollection.cs
@@ -45,6 +45,27 @@ namespace HordeServer.Jobs
 	/// </summary>
 	public class JobCollection : IJobCollection
 	{
+		/// <summary>
+		/// Ease of use struct used to pass telemetry context for job telemetry.
+		/// </summary>
+		internal struct TelemetryContext
+		{
+			/// <summary>
+			/// The job that relates to the telemetry.
+			/// </summary>
+			internal IJob Job { get; set; }
+
+			/// <summary>
+			/// The build config object to use when deriving telemetry sinks.
+			/// </summary>
+			internal IOptionsMonitor<BuildConfig> BuildConfig { get; set; }
+
+			/// <summary>
+			/// The telemetry writer used to publish telemetry.
+			/// </summary>
+			internal ITelemetryWriter TelemetryWriter { get; set; }
+		}
+
 		class Job : IJob
 		{
 			readonly JobCollection _collection;
@@ -783,7 +804,7 @@ namespace HordeServer.Jobs
 		const int MaxRetries = 2;
 
 		readonly IMongoCollection<JobDocument> _jobs;
-		
+
 		[SuppressMessage("CodeQuality", "IDE0052:Remove unread private members")]
 		readonly MongoIndex<JobDocument> _streamThenTemplateThenCreationTimeIndex;
 		readonly MongoIndex<JobDocument> _startedByBisectTaskIdIndex;
@@ -1725,7 +1746,8 @@ namespace HordeServer.Jobs
 			// Update the state of dependent jobsteps
 			if (refreshDependentJobSteps)
 			{
-				RefreshDependentJobSteps(jobDocument, graph, updates, _logger, GetJobLogger(jobDocument.Id));
+				IJob jobObject = CreateJobObject(jobDocument, graph);
+				RefreshDependentJobSteps(jobDocument, graph, updates, _logger, GetJobLogger(jobDocument.Id), new TelemetryContext() { BuildConfig = _buildConfig, TelemetryWriter = _telemetryWriter, Job = jobObject });
 				RefreshJobPriority(jobDocument, updates, _logger, GetJobLogger(jobDocument.Id));
 			}
 
@@ -2078,7 +2100,8 @@ namespace HordeServer.Jobs
 				}
 			}
 
-			RefreshDependentJobSteps(jobDocument, graph, updates, _logger, GetJobLogger(jobDocument.Id));
+			IJob jobObject = CreateJobObject(jobDocument, graph);
+			RefreshDependentJobSteps(jobDocument, graph, updates, _logger, GetJobLogger(jobDocument.Id), new TelemetryContext() { BuildConfig = _buildConfig, TelemetryWriter = _telemetryWriter, Job = jobObject });
 			RefreshJobPriority(jobDocument, updates, _logger, GetJobLogger(jobDocument.Id));
 
 			return await TryUpdateAsync(jobDocument, updateBuilder.Combine(updates), cancellationToken);
@@ -2105,7 +2128,7 @@ namespace HordeServer.Jobs
 		/// <param name="graph">The graph definition for this job</param>
 		/// <param name="logger">Logger for output messages</param>
 		/// <param name="jobAudit"></param>
-		private static void CreateBatches(JobDocument job, IGraph graph, ILogger logger, IAuditLogChannel<JobId> jobAudit)
+		private void CreateBatches(JobDocument job, IGraph graph, ILogger logger, IAuditLogChannel<JobId> jobAudit)
 		{
 			UpdateBatches(job, graph, new List<UpdateDefinition<JobDocument>>(), logger, jobAudit);
 		}
@@ -2118,7 +2141,7 @@ namespace HordeServer.Jobs
 		/// <param name="updates">List of updates for the job</param>
 		/// <param name="logger">Logger for updates</param>
 		/// <param name="jobAudit"></param>
-		static void UpdateBatches(JobDocument job, IGraph graph, List<UpdateDefinition<JobDocument>> updates, ILogger logger, IAuditLogChannel<JobId> jobAudit)
+		void UpdateBatches(JobDocument job, IGraph graph, List<UpdateDefinition<JobDocument>> updates, ILogger logger, IAuditLogChannel<JobId> jobAudit)
 		{
 			UpdateDefinitionBuilder<JobDocument> updateBuilder = Builders<JobDocument>.Update;
 
@@ -2128,7 +2151,8 @@ namespace HordeServer.Jobs
 			updates.Add(updateBuilder.Set(x => x.NextSubResourceId, job.NextSubResourceId));
 
 			// Update all the dependencies
-			RefreshDependentJobSteps(job, graph, new List<UpdateDefinition<JobDocument>>(), logger, jobAudit);
+			IJob jobObject = CreateJobObject(job, graph);
+			RefreshDependentJobSteps(job, graph, new List<UpdateDefinition<JobDocument>>(), logger, jobAudit, new TelemetryContext() { BuildConfig = _buildConfig, TelemetryWriter = _telemetryWriter, Job = jobObject });
 			RefreshJobPriority(job, updates, logger, jobAudit);
 		}
 
@@ -2573,7 +2597,8 @@ namespace HordeServer.Jobs
 		/// <param name="updates">List of updates to the job</param>
 		/// <param name="logger">Logger instance</param>
 		/// <param name="jobAudit"></param>
-		static void RefreshDependentJobSteps(JobDocument job, IGraph graph, List<UpdateDefinition<JobDocument>> updates, ILogger logger, IAuditLogChannel<JobId> jobAudit)
+		/// <param name="telemetryContext">Context object used in telemetry constuction, and emissions.</param>
+		static void RefreshDependentJobSteps(JobDocument job, IGraph graph, List<UpdateDefinition<JobDocument>> updates, ILogger logger, IAuditLogChannel<JobId> jobAudit, TelemetryContext? telemetryContext = null)
 		{
 			// Update the batches
 			UpdateDefinitionBuilder<JobDocument> updateBuilder = Builders<JobDocument>.Update;
@@ -2601,6 +2626,25 @@ namespace HordeServer.Jobs
 							{
 								newState = JobStepState.Skipped;
 								newOutcome = JobStepOutcome.Failure;
+
+								// [TEMPORARY][START] This explicit emission of a skipped job step ref should be removed upon UE-357024.
+								try
+								{
+									if (telemetryContext != null && telemetryContext.Value.BuildConfig.CurrentValue.TryGetStream(job.StreamId, out StreamConfig? streamConfig) && !streamConfig.TelemetryStoreId.IsEmpty)
+									{
+										JobStepRefTelemetry? telemetry = JobStepRefTelemetry.ConstructJobStepRefTelemetryFromJob(telemetryContext.Value.Job, batchIdx, stepIdx, newOutcome, newState);
+
+										if (telemetry != null)
+										{
+											telemetryContext.Value.TelemetryWriter.WriteEvent(streamConfig.TelemetryStoreId, telemetry);
+										}
+									}
+								}
+								catch (Exception ex)
+								{
+									logger.LogError(ex, "Unable to create synthesized job step ref telemetry.");
+								}
+								// [TEMPORARY][END] This explicit emission of a skipped job step ref should be removed upon UE-357024.
 							}
 							else if (!steps.Any(x => !x.AbortRequested && JobStepExtensions.IsPendingState(x.State)))
 							{
diff --git a/Engine/Source/Programs/Horde/Plugins/Build/HordeServer.Build/Jobs/JobStepRefCollection.cs b/Engine/Source/Programs/Horde/Plugins/Build/HordeServer.Build/Jobs/JobStepRefCollection.cs
--- a/Engine/Source/Programs/Horde/Plugins/Build/HordeServer.Build/Jobs/JobStepRefCollection.cs
+++ b/Engine/Source/Programs/Horde/Plugins/Build/HordeServer.Build/Jobs/JobStepRefCollection.cs
@@ -11,6 +11,7 @@ using HordeServer.Server;
 using HordeServer.Streams;
 using HordeServer.Telemetry;
 using HordeServer.Utilities;
+using Microsoft.Extensions.Logging;
 using Microsoft.Extensions.Options;
 using MongoDB.Bson;
 using MongoDB.Bson.Serialization.Attributes;
@@ -156,37 +157,42 @@ namespace HordeServer.Jobs
 		}
 
 		/// <inheritdoc/>
-		public async Task<IJobStepRef> InsertOrReplaceAsync(JobStepRefId id, string jobName, string stepName, StreamId streamId, TemplateId templateId, CommitIdWithOrder commitId, LogId? logId, PoolId? poolId, AgentId? agentId, JobStepState? state, JobStepOutcome? outcome, bool updateIssues, CommitIdWithOrder? lastSuccess, CommitIdWithOrder? lastWarning, float waitTime, float initTime, DateTime jobStartTimeUtc, DateTime startTimeUtc, DateTime? finishTimeUtc, List<string>? stepMetadata = null)
+		public async Task<IJobStepRef> InsertOrReplaceAsync(JobStepRefId id, string jobName, string stepName, StreamId streamId, TemplateId templateId, CommitIdWithOrder commitId, LogId? logId, PoolId? poolId, AgentId? agentId, JobStepState? state, JobStepOutcome? outcome, bool updateIssues, CommitIdWithOrder? lastSuccess, CommitIdWithOrder? lastWarning, float waitTime, float initTime, DateTime jobStartTimeUtc, DateTime startTimeUtc, DateTime? finishTimeUtc, List<string>? stepMetadata = null, ILogger? logger = null)
 		{
 			JobStepRef newJobStepRef = new JobStepRef(id, jobName, stepName, streamId, templateId, commitId, logId, poolId, agentId, state, outcome, updateIssues, lastSuccess, lastWarning, waitTime, initTime, jobStartTimeUtc, startTimeUtc, finishTimeUtc, stepMetadata);
 			await _jobStepRefs.ReplaceOneAsync(Builders<JobStepRef>.Filter.Eq(x => x.Id, newJobStepRef.Id), newJobStepRef, new ReplaceOptions { IsUpsert = true });
 
-			if (_buildConfig.CurrentValue.TryGetStream(streamId, out StreamConfig? streamConfig) && !streamConfig.TelemetryStoreId.IsEmpty)
+			try
 			{
-				_telemetryWriter.WriteEvent(streamConfig.TelemetryStoreId, new
+				if (_buildConfig.CurrentValue.TryGetStream(streamId, out StreamConfig? streamConfig) && !streamConfig.TelemetryStoreId.IsEmpty)
 				{
-					EventName = "State.JobStepRef",
-					Id = id.ToString(),
-					JobId = id.JobId.ToString(),
-					BatchId = id.BatchId.ToString(),
-					StepId = id.StepId.ToString(),
-					AgentId = agentId,
-					BatchInitTime = initTime,
-					BatchWaitTime = waitTime,
-					Change = commitId?.Name,
-					FinishTime = finishTimeUtc,
-					JobName = jobName,
-					JobStartTime = jobStartTimeUtc,
-					StepName = stepName,
-					State = state,
-					Outcome = outcome,
-					PoolId = poolId,
-					StartTime = startTimeUtc,
-					StreamId = streamId,
-					TemplateId = templateId,
-					UpdateIssues = updateIssues,
-					Duration = (finishTimeUtc != null) ? (double?)(finishTimeUtc.Value - startTimeUtc).TotalSeconds : null
-				});
+					_telemetryWriter.WriteEvent(streamConfig.TelemetryStoreId, new JobStepRefTelemetry(
+						id.ToString(),
+						id.JobId.ToString(),
+						id.BatchId.ToString(),
+						id.StepId.ToString(),
+						agentId,
+						initTime,
+						waitTime,
+						commitId?.Name,
+						finishTimeUtc,
+						jobName,
+						jobStartTimeUtc,
+						stepName,
+						state,
+						outcome,
+						poolId,
+						startTimeUtc,
+						streamId,
+						templateId,
+						updateIssues,
+						(finishTimeUtc != null) ? (double?)(finishTimeUtc.Value - startTimeUtc).TotalSeconds : null)
+					);
+				}
+			}
+			catch (Exception ex)
+			{
+				logger?.LogError(ex, "Unable to create job step ref telemetry.");
 			}
 
 			return newJobStepRef;
@@ -248,7 +254,7 @@ namespace HordeServer.Jobs
 			{
 				filter &= filterBuilder.Ne(x => x.Outcome, JobStepOutcome.Failure);
 			}
-			if(!includePreStartAndSkippedStates)
+			if (!includePreStartAndSkippedStates)
 			{
 				filter &= filterBuilder.Ne(x => x.State, JobStepState.Skipped);
 				filter &= filterBuilder.Ne(x => x.State, JobStepState.Ready);
diff --git a/Engine/Source/Programs/Horde/Plugins/Build/HordeServer.Build/Jobs/JobTelemetry.cs b/Engine/Source/Programs/Horde/Plugins/Build/HordeServer.Build/Jobs/JobTelemetry.cs
--- a/Engine/Source/Programs/Horde/Plugins/Build/HordeServer.Build/Jobs/JobTelemetry.cs
+++ b/Engine/Source/Programs/Horde/Plugins/Build/HordeServer.Build/Jobs/JobTelemetry.cs
@@ -2,6 +2,8 @@
 
 using System.ComponentModel.DataAnnotations.Schema;
 using EpicGames.Analytics.Telemetry;
+using EpicGames.Horde.Agents;
+using EpicGames.Horde.Agents.Pools;
 using EpicGames.Horde.Jobs;
 using EpicGames.Horde.Jobs.Graphs;
 using EpicGames.Horde.Jobs.Templates;
@@ -515,4 +517,232 @@ namespace HordeServer.Jobs
 			State = state;
 		}
 	}
+	/// <summary>
+	/// Record used for Job Step Ref telemetry. Models a Job Step Ref.
+	/// </summary>
+	public record JobStepRefTelemetry
+	{
+		/// <summary>
+		/// The telemetry event name.
+		/// </summary>
+		public string EventName { get; init; }
+
+		/// <summary>
+		/// The schema version of the telemetry.
+		/// </summary>
+		public int SchemaVersion { get; init; }
+
+		/// <summary>
+		/// The job step ref id.
+		/// </summary>
+		public string Id { get; init; }
+
+		/// <summary>
+		/// The job id.
+		/// </summary>
+		public string JobId { get; init; }
+
+		/// <summary>
+		/// The batch id.
+		/// </summary>
+		public string BatchId { get; init; }
+
+		/// <summary>
+		/// The step id.
+		/// </summary>
+		public string StepId { get; init; }
+
+		/// <summary>
+		/// The init time of the batch.
+		/// </summary>
+		public float BatchInitTime { get; init; }
+
+		/// <summary>
+		/// The wait time of the batch.
+		/// </summary>
+		public float BatchWaitTime { get; init; }
+
+		/// <summary>
+		/// The commit id.
+		/// </summary>
+		public string? Change { get; init; }
+
+		/// <summary>
+		/// The job name.
+		/// </summary>
+		public string JobName { get; init; }
+
+		/// <summary>
+		/// The step finish time.
+		/// </summary>
+		public DateTime? FinishTime { get; init; }
+
+		/// <summary>
+		/// The owning job's start time.
+		/// </summary>
+		public DateTime? JobStartTime { get; init; }
+
+		/// <summary>
+		/// The step name.
+		/// </summary>
+		public string StepName { get; init; }
+
+		/// <summary>
+		/// The step state.
+		/// </summary>
+		public JobStepState? State { get; init; }
+
+		/// <summary>
+		/// The step outcome.
+		/// </summary>
+		public JobStepOutcome? Outcome { get; init; }
+
+		/// <summary>
+		/// The pool id of the associated agent.
+		/// </summary>
+		public PoolId? PoolId { get; init; }
+
+		/// <summary>
+		/// The start time of the step.
+		/// </summary>
+		public DateTime? StartTime { get; init; }
+
+		/// <summary>
+		/// The stream id.
+		/// </summary>
+		public StreamId StreamId { get; init; }
+
+		/// <summary>
+		/// The template id of the owning job.
+		/// </summary>
+		public TemplateId TemplateId { get; init; }
+
+		/// <summary>
+		/// The id of the agent that ran the step.
+		/// </summary>
+		public AgentId? AgentId { get; init; }
+
+		/// <summary>
+		/// Whether issues have been updated or not.
+		/// </summary>
+		public bool UpdateIssues { get; init; }
+
+		/// <summary>
+		/// The duration of the step.
+		/// </summary>
+		public double? Duration { get; init; }
+
+		/// <summary>
+		/// The default event name.
+		/// </summary>
+		public const string DefaultEventName = "State.JobStepRef";
+
+		/// <summary>
+		/// Current schema version.
+		/// </summary>
+		public static readonly int CurrentSchemaVersion = 2;
+
+		/// <summary>
+		/// Constructor.
+		/// </summary>
+		/// <param name="id">The job step ref id.</param>
+		/// <param name="jobId">The job id.</param>
+		/// <param name="batchId">The batch id.</param>
+		/// <param name="stepId">The step id.</param>
+		/// <param name="agentId">The id of the agent that ran the step.</param>
+		/// <param name="batchInitTime">The init time of the batch.</param>
+		/// <param name="batchWaitTime">The wait time of the batch.</param>
+		/// <param name="change">The commit id.</param>
+		/// <param name="finishTime">The step finish time.</param>
+		/// <param name="jobName">The job name.</param>
+		/// <param name="jobStartTime">The owning job's start time.</param>
+		/// <param name="stepName">The step name.</param>
+		/// <param name="state">The step state.</param>
+		/// <param name="outcome">The step outcome.</param>
+		/// <param name="poolId">The pool id of the associated agent.</param>
+		/// <param name="startTime">The start time of the step.</param>
+		/// <param name="streamId">The stream id.</param>
+		/// <param name="templateId">The template id of the owning job.</param>
+		/// <param name="updateIssues">Whether issues have been updated or not.</param>
+		/// <param name="duration">The duration of the step.</param>
+		public JobStepRefTelemetry(string id, string jobId, string batchId, string stepId, AgentId? agentId, float batchInitTime, float batchWaitTime, string? change, DateTime? finishTime, string jobName, DateTime? jobStartTime, string stepName, JobStepState? state, JobStepOutcome? outcome, PoolId? poolId, DateTime? startTime, StreamId streamId, TemplateId templateId, bool updateIssues, double? duration)
+		{
+			EventName = DefaultEventName;
+			SchemaVersion = CurrentSchemaVersion;
+			Id = id;
+			JobId = jobId;
+			BatchId = batchId;
+			StepId = stepId;
+			BatchInitTime = batchInitTime;
+			BatchWaitTime = batchWaitTime;
+			Change = change;
+			FinishTime = finishTime;
+			JobName = jobName;
+			JobStartTime = jobStartTime;
+			StepName = stepName;
+			State = state;
+			Outcome = outcome;
+			PoolId = poolId;
+			StartTime = startTime;
+			StreamId = streamId;
+			TemplateId = templateId;
+			AgentId = agentId;
+			UpdateIssues = updateIssues;
+			Duration = duration;
+		}
+
+		/// <summary>
+		/// Create method used to construct a JobStepRef telemetry event from a job, the provided batch index, and step index. This will take a new outcome and state to use in the telemetry emission.
+		/// </summary>
+		/// <param name="job">The job to use in constructing the telemetry event.</param>
+		/// <param name="batchIdx">The index of the job batch (into the job's batches) to use in constructing the telemetry event.</param>
+		/// <param name="stepIdx">The index of the job step (into the batch's steps) to use in constructing the telemetry event.</param>
+		/// <param name="newOutcome">The job step ref telemetry event's step outcome.</param>
+		/// <param name="newStepState">The job step ref telemetry event's step state.</param>
+		/// <returns>The JobStepRefTelemetry event, null if unable to construct due to invalid parameters.</returns>
+		/// <remarks>Should only be used until full state replication is completed in UE-357021.</remarks>
+		public static JobStepRefTelemetry? ConstructJobStepRefTelemetryFromJob(IJob job, int batchIdx, int stepIdx, JobStepOutcome newOutcome, JobStepState newStepState)
+		{
+			if (batchIdx < 0 || job.Batches.Count <= batchIdx)
+			{
+				return null;
+			}
+
+			IJobStepBatch batch = job.Batches[batchIdx];
+
+			if (stepIdx < 0 || batch.Steps.Count <= stepIdx)
+			{
+				return null;
+			}
+
+			IJobStep step = batch.Steps[stepIdx];
+			JobStepRefId jobStepRefId = new JobStepRefId(job.Id, batch.Id, step.Id);
+			float waitTime = (float)(batch.GetWaitTime() ?? TimeSpan.Zero).TotalSeconds;
+			float initTime = (float)(batch.GetInitTime() ?? TimeSpan.Zero).TotalSeconds;
+
+			JobStepRefTelemetry telemetry = new JobStepRefTelemetry(
+				jobStepRefId.ToString(),
+				job.Id.ToString(),
+				batch.Id.ToString(),
+				step.Id.ToString(),
+				batch.AgentId,
+				initTime,
+				waitTime,
+				job.CommitId?.Name ?? "0",
+				step.FinishTimeUtc,
+				job.Name,
+				job.CreateTimeUtc,
+				step.Name,
+				newStepState,
+				newOutcome,
+				batch.PoolId,
+				step.StartTimeUtc,
+				job.StreamId,
+				job.TemplateId,
+				job.UpdateIssues,
+				(step.FinishTimeUtc != null && step.StartTimeUtc != null) ? (double)(step.FinishTimeUtc.Value - step.StartTimeUtc.Value).TotalSeconds : null);
+
+			return telemetry;
+		}
+	}
 }
\ No newline at end of file
