Introduce Hourly cadence badge

#rnx

#Synced-CL 49105898

diff --git a/Engine/Source/Programs/Horde/HordeDashboard/plugins/test-automation/testAutomationHealth.tsx b/Engine/Source/Programs/Horde/HordeDashboard/plugins/test-automation/testAutomationHealth.tsx
--- a/Engine/Source/Programs/Horde/HordeDashboard/plugins/test-automation/testAutomationHealth.tsx
+++ b/Engine/Source/Programs/Horde/HordeDashboard/plugins/test-automation/testAutomationHealth.tsx
@@ -403,9 +403,12 @@ const getTestHealthItems = (handler: TestDataHandler, stream?: string, test?: Te
             if (highestRunCount > 0) {
                 const rangeDayCount = handler.filterState.weeks * 7;
                 const cadenceRate = highestRunCount / rangeDayCount;
-                if (cadenceRate > 2) {
-                    /** more than 2 runs / 1 day */
+                if (cadenceRate > 24) {
+                    /** more than 24 runs / 1 day */
                     cadenceValue = TestCadence.Incremental;
+                } else if (cadenceRate > 3) {
+                    /** more than 3 runs / 1 day */
+                    cadenceValue = TestCadence.Hourly;
                 } else if (cadenceRate > 0.4) {
                     /** more than 3 runs / 7 days */
                     cadenceValue = TestCadence.Daily;
diff --git a/Engine/Source/Programs/Horde/HordeDashboard/plugins/test-automation/testHealthData.ts b/Engine/Source/Programs/Horde/HordeDashboard/plugins/test-automation/testHealthData.ts
--- a/Engine/Source/Programs/Horde/HordeDashboard/plugins/test-automation/testHealthData.ts
+++ b/Engine/Source/Programs/Horde/HordeDashboard/plugins/test-automation/testHealthData.ts
@@ -36,6 +36,7 @@ export const healthColors: Map<TestHealth, string> = new Map([
 export enum TestCadence {
     Nightly = "Nightly",
     Incremental = "Incrementally",
+    Hourly = "Hourly",
     Daily = "Daily",
     Limited = "Limited",
     Weekly = "Weekly",
@@ -54,6 +55,7 @@ export const cadenceOrders: Map<TestCadence, number> = new Map([
 
 export const cadenceColors: Map<TestCadence, string> = new Map([
     [TestCadence.Incremental, "#c8b72081"],
+    [TestCadence.Hourly, "#9fce538d"],
     [TestCadence.Daily, "#53cea18d"],
     [TestCadence.Weekly, "#1f49a579"],
     [TestCadence.Limited, "#4a3e6c8b"],
