[TestData] Enforce descending sort order in test and phase sessions

#rnx
#Synced-CL 49105898

diff --git a/Engine/Source/Programs/Horde/Plugins/Build/HordeServer.Build/Jobs/TestData/TestDataCollectionV2.cs b/Engine/Source/Programs/Horde/Plugins/Build/HordeServer.Build/Jobs/TestData/TestDataCollectionV2.cs
--- a/Engine/Source/Programs/Horde/Plugins/Build/HordeServer.Build/Jobs/TestData/TestDataCollectionV2.cs
+++ b/Engine/Source/Programs/Horde/Plugins/Build/HordeServer.Build/Jobs/TestData/TestDataCollectionV2.cs
@@ -800,6 +800,7 @@ namespace HordeServer.Jobs.TestData
 		{
 			FilterDefinition<TestSessionDocument> filter = FilterDefinition<TestSessionDocument>.Empty;
 			FilterDefinitionBuilder<TestSessionDocument> filterBuilder = Builders<TestSessionDocument>.Filter;
+			SortDefinition<TestSessionDocument> sortDefinition = Builders<TestSessionDocument>.Sort.Descending(x => x.BuildCommitOrder);
 
 			filter &= filterBuilder.In(x => x.StreamId, streamIds);
 
@@ -840,7 +841,7 @@ namespace HordeServer.Jobs.TestData
 
 			using (TelemetrySpan _ = _tracer.StartActiveSpan($"{nameof(TestDataCollectionV2)}.{nameof(FindTestSessionsAsync)}"))
 			{
-				results = await _testSessions.Find(filter).ToListAsync(cancellationToken);
+				results = await _testSessions.Find(filter).Sort(sortDefinition).ToListAsync(cancellationToken);
 			}
 
 			return results.ConvertAll<ITestSession>(x => x);
@@ -952,6 +953,7 @@ namespace HordeServer.Jobs.TestData
 		{
 			FilterDefinition<TestPhaseSessionDocument> filter = FilterDefinition<TestPhaseSessionDocument>.Empty;
 			FilterDefinitionBuilder<TestPhaseSessionDocument> filterBuilder = Builders<TestPhaseSessionDocument>.Filter;
+			SortDefinition<TestPhaseSessionDocument> sortDefinition = Builders<TestPhaseSessionDocument>.Sort.Descending(x => x.BuildCommitOrder);
 
 			filter &= filterBuilder.In(x => x.StreamId, streamIds);
 			filter &= filterBuilder.In(x => x.PhaseRef, phaseIds);
@@ -983,7 +985,7 @@ namespace HordeServer.Jobs.TestData
 
 			using (TelemetrySpan _ = _tracer.StartActiveSpan($"{nameof(TestDataCollectionV2)}.{nameof(FindTestPhaseSessionsAsync)}"))
 			{
-				results = await _testPhaseSessions.Find(filter).ToListAsync(cancellationToken);
+				results = await _testPhaseSessions.Find(filter).Sort(sortDefinition).ToListAsync(cancellationToken);
 			}
 
 			return results.ConvertAll<ITestPhaseSession>(x => x);
