[Horde][BuildHealth][Dashboard ]Add ability to export basic query params for new build

#Synced-CL 0

diff --git a/Engine/Source/Programs/Horde/HordeDashboard/src/base/utilities/clipboard.ts b/Engine/Source/Programs/Horde/HordeDashboard/src/base/utilities/clipboard.ts
--- a/Engine/Source/Programs/Horde/HordeDashboard/src/base/utilities/clipboard.ts
+++ b/Engine/Source/Programs/Horde/HordeDashboard/src/base/utilities/clipboard.ts
@@ -13,3 +13,21 @@ export const copyToClipboard = (value: string | undefined) => {
     document.execCommand('copy');
     document.body.removeChild(el);
 }
+
+// New version of copy to clipboard utility function, using modern method and returning results, to enable accurate feedback in the UI.
+// Should eventually replace copyToClipboard.
+export const copyToClipboardAsync = async (value: string | undefined): Promise<boolean> => {
+
+    if (!value) {
+        return false;
+    }
+
+    try {
+        await navigator.clipboard.writeText(value);
+    } catch (reason) {
+        console.log(`Failed to copy '${value}' to clipboard: ${reason}`);
+        return false;
+    }
+
+    return true;
+}
diff --git a/Engine/Source/Programs/Horde/HordeDashboard/src/components/build/NewBuildV2.tsx b/Engine/Source/Programs/Horde/HordeDashboard/src/components/build/NewBuildV2.tsx
--- a/Engine/Source/Programs/Horde/HordeDashboard/src/components/build/NewBuildV2.tsx
+++ b/Engine/Source/Programs/Horde/HordeDashboard/src/components/build/NewBuildV2.tsx
@@ -12,7 +12,7 @@ import { BoolParameterData, ChangeQueryConfig, CreateJobRequest, GetJobsTabRespo
 import dashboard from "../../backend/Dashboard";
 import { ProjectStore } from "../../backend/ProjectStore";
 import templateCache from "../../backend/TemplateCache";
-import { copyToClipboard } from "../../base/utilities/clipboard";
+import { copyToClipboard, copyToClipboardAsync } from "../../base/utilities/clipboard";
 import { getHordeStyling } from "../../styles/Styles";
 import { errorDialogStore } from "../error/ErrorStore";
 import { NewBuild } from "../NewBuild";
@@ -67,20 +67,20 @@ class BuildOptions {
 
    }
 
-   getQuotedArg(arg: string){
+   getQuotedArg(arg: string) {
       // Ignore args that don't follow "-set:XXX=YYY"
       // Also prevents overquoting incorrectly recorded args 
       // i.e., forgetting "" in -set:XXX=YYY ZZZ AAA BBB => -set:XXX="YYY" "ZZZ" "AAA" "BBB"
-      if(!arg.startsWith("-set")) return arg;
+      if (!arg.startsWith("-set")) return arg;
 
       // Replace any existing quotes with escaped quotes
       arg = arg.replace(/"/g, '\\"');
-      
+
       // Split argument on space to extract left side assignment
       const splitOnSpace = arg.split(" ");
-      
+
       // Ignore args without spaces
-      if(splitOnSpace.length === 1) return arg;
+      if (splitOnSpace.length === 1) return arg;
 
       // Split assignment on equals, getting the last = sign. Which equals sign doesn't necessarily matter.
       // It looks best/is more semantic to do the inner most one, though.
@@ -88,11 +88,11 @@ class BuildOptions {
 
       // Add double quote just before value
       splitOnEquals[splitOnEquals.length - 1] = `"${splitOnEquals[splitOnEquals.length - 1]}`;
-      
+
       // Replace splitOnSpace[0], Restore string, and add " to end
       splitOnSpace[0] = splitOnEquals.join("=");
       let finalString = splitOnSpace.join(" ") + '"';
-      
+
       // -set:XXX=YYY=AAA BBB CCC --> -set:XXX=YYY="AAA BBB CCC"
       return finalString;
    }
@@ -1911,8 +1911,52 @@ const AdvancedPanel: React.FC = observer(() => {
 
 })
 
+async function buildParametersUrl(template: GetTemplateRefResponse): Promise<boolean> {
+
+   // Note: can't rely on current path for template as user can change that via All Templates and URL won't update.
+   const setParams = new URLSearchParams(`${location.search}&newbuild=true&templateId=${template.id}`);
+   const args = BuildOptions.get().parameters;
+
+   for (const p of template.parameters) {
+      switch (p.type) {
+         case ParameterType.Bool: {
+            const b = p as BoolParameterData;
+            if (args[b.id] && args[b.id] !== b.default.toString()) {
+               setParams.set(`id-${b.id}`, args[b.id]);
+            }
+            break;
+         }
+         case ParameterType.Text: {
+            const t = p as TextParameterData;
+            if (args[t.id] && args[t.id] !== t.default) {
+               setParams.set(`id-${t.id}`, args[t.id]);
+            }
+            break;
+         }
+         case ParameterType.List: {
+            const list = p as ListParameterData;
+            for (const i of list.items) {
+               // Note: need to include only true values for lists, as setDefaultParameters clears all values if any are present in the url.
+               if (args[i.id] && args[i.id].toLowerCase() === "true") {
+                  setParams.set(`id-${i.id}`, args[i.id]);
+               }
+            }
+            break;
+         }
+         default:
+            console.error(`Unknown parameter type ${p.type}`);
+      }
+   }
+   
+   const query = new URL(window.location.href);
+   query.search = setParams.toString();
+   return await copyToClipboardAsync(query.toString());
+}
+
 const BuildModal: React.FC<{ setUseLegacyDialog: (value: boolean) => void }> = observer(({ setUseLegacyDialog }) => {
 
+   const [parametersUrlCopied, setParametersUrlCopied] = useState(false);
+
    const options = BuildOptions.get();
    options.subscribe();
 
@@ -2182,8 +2226,16 @@ const BuildModal: React.FC<{ setUseLegacyDialog: (value: boolean) => void }> = o
                </Stack>
                <Stack horizontal tokens={{ childrenGap: 16 }} styles={{ root: { paddingTop: 32, paddingLeft: 8, paddingBottom: 8 } }}>
                   {!options.jobDetails && !options.readOnly && <Stack>
-                     <DefaultButton text="Reset" style={{ width: 160 }} onClick={() => options.onResetToDefaults()} />
+                     <DefaultButton text="Reset" style={{ width: 110 }} onClick={() => options.onResetToDefaults()} />
                   </Stack>}
+                  <DefaultButton text={parametersUrlCopied ? "URL copied!": "Create URL"} disabled={!template}
+                     onClick={async () => { 
+                        const copySuccess = await buildParametersUrl(template);
+                        setParametersUrlCopied(copySuccess);
+                        if (copySuccess) {
+                           setTimeout(() => setParametersUrlCopied(false), 5000);
+                        }
+                  }} />
 
                   <Stack grow />
                   <PrimaryButton text="Start Job" disabled={!template || !template.canRun || options.readOnly || options.preview} onClick={() => { onSubmit(); }} />
