[Horde][Analytics][Server] Using nullable type for IJob as per CreateJobObject return specification.

#Synced-CL 49094299

diff --git a/Engine/Source/Programs/Horde/Plugins/Build/HordeServer.Build/Jobs/JobCollection.cs b/Engine/Source/Programs/Horde/Plugins/Build/HordeServer.Build/Jobs/JobCollection.cs
--- a/Engine/Source/Programs/Horde/Plugins/Build/HordeServer.Build/Jobs/JobCollection.cs
+++ b/Engine/Source/Programs/Horde/Plugins/Build/HordeServer.Build/Jobs/JobCollection.cs
@@ -53,7 +53,7 @@ namespace HordeServer.Jobs
 			/// <summary>
 			/// The job that relates to the telemetry.
 			/// </summary>
-			internal IJob Job { get; set; }
+			internal IJob? Job { get; set; }
 
 			/// <summary>
 			/// The build config object to use when deriving telemetry sinks.
@@ -1746,7 +1746,7 @@ namespace HordeServer.Jobs
 			// Update the state of dependent jobsteps
 			if (refreshDependentJobSteps)
 			{
-				IJob jobObject = CreateJobObject(jobDocument, graph);
+				IJob? jobObject = CreateJobObject(jobDocument, graph);
 				RefreshDependentJobSteps(jobDocument, graph, updates, _logger, GetJobLogger(jobDocument.Id), new TelemetryContext() { BuildConfig = _buildConfig, TelemetryWriter = _telemetryWriter, Job = jobObject });
 				RefreshJobPriority(jobDocument, updates, _logger, GetJobLogger(jobDocument.Id));
 			}
@@ -2100,7 +2100,7 @@ namespace HordeServer.Jobs
 				}
 			}
 
-			IJob jobObject = CreateJobObject(jobDocument, graph);
+			IJob? jobObject = CreateJobObject(jobDocument, graph);
 			RefreshDependentJobSteps(jobDocument, graph, updates, _logger, GetJobLogger(jobDocument.Id), new TelemetryContext() { BuildConfig = _buildConfig, TelemetryWriter = _telemetryWriter, Job = jobObject });
 			RefreshJobPriority(jobDocument, updates, _logger, GetJobLogger(jobDocument.Id));
 
@@ -2151,7 +2151,7 @@ namespace HordeServer.Jobs
 			updates.Add(updateBuilder.Set(x => x.NextSubResourceId, job.NextSubResourceId));
 
 			// Update all the dependencies
-			IJob jobObject = CreateJobObject(job, graph);
+			IJob? jobObject = CreateJobObject(job, graph);
 			RefreshDependentJobSteps(job, graph, new List<UpdateDefinition<JobDocument>>(), logger, jobAudit, new TelemetryContext() { BuildConfig = _buildConfig, TelemetryWriter = _telemetryWriter, Job = jobObject });
 			RefreshJobPriority(job, updates, logger, jobAudit);
 		}
diff --git a/Engine/Source/Programs/Horde/Plugins/Build/HordeServer.Build/Jobs/JobTelemetry.cs b/Engine/Source/Programs/Horde/Plugins/Build/HordeServer.Build/Jobs/JobTelemetry.cs
--- a/Engine/Source/Programs/Horde/Plugins/Build/HordeServer.Build/Jobs/JobTelemetry.cs
+++ b/Engine/Source/Programs/Horde/Plugins/Build/HordeServer.Build/Jobs/JobTelemetry.cs
@@ -202,7 +202,7 @@ namespace HordeServer.Jobs
 			IsPreflight = false;
 			StepTotalCount = 0;
 			StepPassCount = 0;
-			StepWarningCount= 0;
+			StepWarningCount = 0;
 			StepFailureCount = 0;
 			JobStepsTotalTime = 0;
 			PassRatio = 0;
@@ -701,8 +701,13 @@ namespace HordeServer.Jobs
 		/// <param name="newStepState">The job step ref telemetry event's step state.</param>
 		/// <returns>The JobStepRefTelemetry event, null if unable to construct due to invalid parameters.</returns>
 		/// <remarks>Should only be used until full state replication is completed in UE-357021.</remarks>
-		public static JobStepRefTelemetry? ConstructJobStepRefTelemetryFromJob(IJob job, int batchIdx, int stepIdx, JobStepOutcome newOutcome, JobStepState newStepState)
+		public static JobStepRefTelemetry? ConstructJobStepRefTelemetryFromJob(IJob? job, int batchIdx, int stepIdx, JobStepOutcome newOutcome, JobStepState newStepState)
 		{
+			if (job == null)
+			{
+				return null;
+			}
+
 			if (batchIdx < 0 || job.Batches.Count <= batchIdx)
 			{
 				return null;
