Horde: Update release notes

#rnx

diff --git a/Engine/Source/Programs/Horde/Docs/ReleaseNotes.md b/Engine/Source/Programs/Horde/Docs/ReleaseNotes.md
--- a/Engine/Source/Programs/Horde/Docs/ReleaseNotes.md
+++ b/Engine/Source/Programs/Horde/Docs/ReleaseNotes.md
@@ -1,5 +1,29 @@
 # Release Notes
 
+## 2025-12-11
+
+* Using nullable type for IJob as per CreateJobObject return specification. (49147810)
+* Adding telemetry Post api for clients (49140339)
+* Add ability to export basic query params for new build (49139846)
+* JobStepRef telemetry emissions for skipped steps. (49134073)
+* Add metadata to Analytics (49121215)
+* Fixing default configs for Analytics Project. (49115652)
+* JobDriver Perforce updates
+  * Ensure PerforceMaterializer has support for and uses maxFileConcurrency and maxFileScanConcurrency as the ManagedWorkspace, ensures cleaning can run file operations in parallel
+  * PerforceExecutor - conforms now obey the maxFileConcurrency option to ensure deletes are run in parallel as well (49107432)
+* New Analytics projects (49087766)
+* Time formatting pass (user pref based - UTC & 24h clock) (49087746)
+* Adding configurable User1, 2, 3 status colors (49082227)
+* Suppress telemetry events with duplicate keys (49049332)
+* Fix situation where credentials would be printed to log on Mac/Linux because we compared conformed environment variable names against non-conformed credential key names.  Now doing the comparison with conformed values on both sides of the comparison. (49049136)
+* Add append option for retrieving cloud artifacts (49046639)
+* Ensure untracked files are removed in PerforceMaterializer (48991402)
+* Add an EpicGames.Build.Tests project with tests for the MSBuild project parsing logic. (48985777)
+* Pass along old project id to prevent full refresh. (48967771)
+* More polish on default heights for right top edge alignment (48964876)
+* Minor polish for sticky & more appropriate default height. (48961671)
+* Initialize PerforceMetadataLogger even if retrieved after a sync (48953485)
+
 ## 2025-12-04
 
 * Initialize PerforceMetadataLogger even if retrieved after a sync (48953485)
