[Test Health] Rename Workflow Notification to Degradation Notification

#rnx
#Synced-CL 49164840

diff --git a/Engine/Source/Programs/Horde/HordeDashboard/plugins/test-automation/testAutomationHealth.tsx b/Engine/Source/Programs/Horde/HordeDashboard/plugins/test-automation/testAutomationHealth.tsx
--- a/Engine/Source/Programs/Horde/HordeDashboard/plugins/test-automation/testAutomationHealth.tsx
+++ b/Engine/Source/Programs/Horde/HordeDashboard/plugins/test-automation/testAutomationHealth.tsx
@@ -680,7 +680,7 @@ export const TestHealthView: React.FC<{ test: TestHealthItem, handler: TestDataH
 
                                 <Stack horizontal>
                                     <Stack horizontal verticalAlign="center" tokens={{childrenGap: 6}} style={{cursor: 'default'}}>
-                                        <Label style={{whiteSpace: 'nowrap'}}>Workflow Notification: </Label>
+                                        <Label style={{whiteSpace: 'nowrap'}}>Degradation Notification: </Label>
                                         {!notificationOptions && <Text style={{color:textEmphasisColor(), cursor: 'default'}}>no workflow available in this stream</Text>}
                                         {!!notificationOptions && <ComboBox
                                             key={`notification-${test.updated}`}
