Horde: Fix exception during PerforceMaterializer clean when in dirty state

Parallel directory deletion did assume workspace directory exists, which is not always true.

diff --git a/Engine/Source/Programs/Horde/Drivers/JobDriver/Execution/PerforceMaterializer.cs b/Engine/Source/Programs/Horde/Drivers/JobDriver/Execution/PerforceMaterializer.cs
--- a/Engine/Source/Programs/Horde/Drivers/JobDriver/Execution/PerforceMaterializer.cs
+++ b/Engine/Source/Programs/Horde/Drivers/JobDriver/Execution/PerforceMaterializer.cs
@@ -304,8 +304,13 @@ public sealed class PerforceMaterializer : IWorkspaceMaterializer
 			using TelemetrySpan deleteSpan = _tracer.StartActiveSpan($"{nameof(PerforceMaterializer)}.{nameof(SyncAsync)}.Delete");
 			using ILoggerProgress status = _logger.BeginProgressScope("Deleting local workspace files...");
 			Stopwatch deleteTimer = Stopwatch.StartNew();
-			Parallel.ForEach(DirectoryReference.EnumerateFiles(_workspaceDir, "*", SearchOption.AllDirectories), new() { MaxDegreeOfParallelism = _options.MaxFileConcurrency }, FileUtils.ForceDeleteFile);
-			FileUtils.ForceDeleteDirectory(_workspaceDir);
+
+			if (Directory.Exists(_workspaceDir.FullName))
+			{
+				Parallel.ForEach(DirectoryReference.EnumerateFiles(_workspaceDir, "*", SearchOption.AllDirectories), new() { MaxDegreeOfParallelism = _options.MaxFileConcurrency }, FileUtils.ForceDeleteFile);
+				FileUtils.ForceDeleteDirectory(_workspaceDir);
+			}
+			Directory.CreateDirectory(_workspaceDir.FullName);
 			status.Progress = $"({deleteTimer.Elapsed.TotalSeconds:0.0}s)";
 		}
 
