Horde: Add a runtime watchdog to agent to monitor for GC or thread pool stalls

diff --git a/Engine/Source/Programs/Horde/HordeAgent/AgentApp.cs b/Engine/Source/Programs/Horde/HordeAgent/AgentApp.cs
--- a/Engine/Source/Programs/Horde/HordeAgent/AgentApp.cs
+++ b/Engine/Source/Programs/Horde/HordeAgent/AgentApp.cs
@@ -3,6 +3,7 @@
 using System.Diagnostics;
 using System.Diagnostics.CodeAnalysis;
 using System.Reflection;
+using System.Runtime;
 using System.Runtime.InteropServices;
 using System.Text.Json;
 using System.Text.Json.Serialization;
@@ -146,6 +147,7 @@ namespace HordeAgent
 					logger.LogInformation("Config source: {ConfigSource}", configSource);
 				}
 				logger.LogInformation("Logs dir: {LogsDir}", settings.LogsDir);
+				LogGcStats(logger);
 			}
 			
 			IServiceCollection services = new ServiceCollection();
@@ -201,6 +203,8 @@ namespace HordeAgent
 			services.AddSingleton<GrpcService>();
 			services.AddSingleton<TelemetryService>();
 			services.AddHostedService(sp => sp.GetRequiredService<TelemetryService>());
+			services.AddSingleton<RuntimeWatchdog>();
+			services.AddHostedService(sp => sp.GetRequiredService<RuntimeWatchdog>());
 
 			services.AddSingleton<JobHandlerFactory>();
 			services.AddSingleton<StatusService>();
@@ -386,5 +390,17 @@ namespace HordeAgent
 				.Select(path => new FileReference(path))
 				.ToList();
 		}
+
+		internal static void LogGcStats(ILogger logger)
+		{
+			string concurrentGc = "<unknown>";
+			if (GC.GetConfigurationVariables().TryGetValue("ConcurrentGC", out object? concurrentGcObj))
+			{
+				concurrentGc = concurrentGcObj?.ToString() ?? "null";
+			}
+			
+			logger.LogInformation("GC: ServerGc={ServerGc} LatencyMode={LatencyMode} ConcurrentGc={ConcurrentGc}",
+				GCSettings.IsServerGC, GCSettings.LatencyMode, concurrentGc);
+		}
 	}
 }
diff --git a/Engine/Source/Programs/Horde/HordeAgent/Commands/Utilities/WatchdogCommand.cs b/Engine/Source/Programs/Horde/HordeAgent/Commands/Utilities/WatchdogCommand.cs
new file mode 100644
--- /dev/null
+++ b/Engine/Source/Programs/Horde/HordeAgent/Commands/Utilities/WatchdogCommand.cs
@@ -0,0 +1,80 @@
+// Copyright Epic Games, Inc. All Rights Reserved.
+
+using System.Text;
+using EpicGames.Core;
+using HordeAgent.Utility;
+using Microsoft.Extensions.Logging;
+
+namespace HordeAgent.Commands.Utilities;
+
+/// <summary>
+/// Test runtime watchdog
+/// </summary>
+[Command("watchdog", "Test runtime watchdog")]
+class RuntimeWatchdogCommand : Command
+{
+	/// <summary>
+	/// Constructor
+	/// </summary>
+	public RuntimeWatchdogCommand()
+	{
+	}
+
+	/// <inheritdoc/>
+	public override async Task<int> ExecuteAsync(ILogger logger)
+	{
+		using ILoggerFactory loggerFactory = LoggerFactory.Create(builder =>
+		{
+			builder.SetMinimumLevel(LogLevel.Debug);
+			builder.AddSimpleConsole(options => { options.SingleLine = true; });
+		});
+		ILogger<RuntimeWatchdog> rwLogger = loggerFactory.CreateLogger<RuntimeWatchdog>();
+		using RuntimeWatchdog watchdog = new (rwLogger);
+		await watchdog.StartAsync(CancellationToken.None);
+
+		watchdog.ArtificialThreadDelay = TimeSpan.FromSeconds(10);
+		StarveWorkerThreads();
+		await Task.Delay(120 * 1000);
+
+		return 0;
+	}
+	
+	internal static void StarveWorkerThreads()
+	{
+		if (!ThreadPool.SetMinThreads(1, 1))
+		{
+			throw new Exception("Failed to set min threads");
+		}
+		if (!ThreadPool.SetMaxThreads(3, 3))
+		{
+			throw new Exception("Failed to set max threads");
+		}
+
+		ThreadPool.GetMaxThreads(out int maxWorker, out int _);
+		for (int i = 0; i < maxWorker + 5; i++)
+		{
+			ThreadPool.QueueUserWorkItem(_ =>
+			{
+				LogThreadPoolState();
+				Thread ct = Thread.CurrentThread;
+				Console.WriteLine($"Simulating blocking work in thread {ct.ManagedThreadId} {ct.Name}");
+				Thread.Sleep(TimeSpan.FromMinutes(5));
+			});
+		}
+	}
+	
+	private static void LogThreadPoolState()
+	{
+		ThreadPool.GetMinThreads(out int workerMin, out int comPortMin);
+		ThreadPool.GetMaxThreads(out int workerMax, out int comPortMax);
+		ThreadPool.GetAvailableThreads(out int workerFree, out int comPortFree);
+		StringBuilder sb = new();
+		sb.Append("Thread ");
+		sb.AppendFormat("Count={0} ", ThreadPool.ThreadCount);
+		sb.AppendFormat("PendingItems={0} ", ThreadPool.PendingWorkItemCount);
+		sb.AppendFormat("CompletedItems={0} | ", ThreadPool.CompletedWorkItemCount);
+		sb.AppendFormat("Worker min={0} max={1} free={2} | ", workerMin, workerMax, workerFree);
+		sb.AppendFormat("ComPort min={0} max={1} free={2} ", comPortMin, comPortMax, comPortFree);
+		Console.WriteLine(sb.ToString());
+	}
+}
diff --git a/Engine/Source/Programs/Horde/HordeAgent/Utility/RuntimeWatchdog.cs b/Engine/Source/Programs/Horde/HordeAgent/Utility/RuntimeWatchdog.cs
new file mode 100644
--- /dev/null
+++ b/Engine/Source/Programs/Horde/HordeAgent/Utility/RuntimeWatchdog.cs
@@ -0,0 +1,140 @@
+﻿// Copyright Epic Games, Inc. All Rights Reserved.
+
+using Microsoft.Extensions.Hosting;
+using Microsoft.Extensions.Logging;
+
+namespace HordeAgent.Utility;
+
+/// <summary>
+/// Monitor execution of the process and .NET runtime
+/// Horde agent runs in an environment where resources are often maximally utilized (CPU, memory, disk, network etc)
+/// This can lead to unexpected behaviors with delays or timeouts.
+/// This class self-monitors and ensures a heartbeat thread and task is "on time".
+/// </summary>
+public sealed class RuntimeWatchdog : IHostedService, IDisposable
+{
+	internal TimeSpan? ArtificialThreadDelay { get; set; } = null;
+	internal TimeSpan? ArtificialTaskDelay { get; set; } = null;
+	
+	private long _lastTaskHeartbeatTicks;
+	private readonly Thread _heartbeatThread;
+	private Task? _heartbeatTask;
+	private volatile bool _isRunning;
+	private readonly TimeSpan _heartbeatThreadInterval = TimeSpan.FromSeconds(10);
+	private readonly TimeSpan _heartbeatTaskInterval = TimeSpan.FromSeconds(5);
+	private readonly TimeSpan _maxTaskStallDuration;
+	private readonly TimeSpan _maxThreadStallDuration;
+	private readonly TimeSpan _maxGcPauseDelta = TimeSpan.FromSeconds(3);
+	private readonly ILogger<RuntimeWatchdog> _logger;
+	private readonly CancellationTokenSource _cts = new();
+
+	/// <summary>
+	/// Constructor
+	/// </summary>
+	public RuntimeWatchdog(ILogger<RuntimeWatchdog> logger)
+	{
+		_maxThreadStallDuration = _heartbeatThreadInterval + TimeSpan.FromSeconds(5);
+		_maxTaskStallDuration = _heartbeatTaskInterval + TimeSpan.FromSeconds(5);
+		_lastTaskHeartbeatTicks = DateTime.UtcNow.Ticks;
+		_heartbeatThread = new Thread(HeartbeatThreadLoop) { Name = "RuntimeWatchdog", IsBackground = true, Priority =  ThreadPriority.AboveNormal };
+		_logger = logger;
+	}
+	
+	/// <inheritdoc/>
+	public void Dispose()
+	{
+		_cts.Dispose();
+	}
+	
+	/// <inheritdoc/>
+	public Task StartAsync(CancellationToken cancellationToken)
+	{
+		if (!_isRunning)
+		{
+			_isRunning = true;
+			_lastTaskHeartbeatTicks = DateTime.UtcNow.Ticks;
+			_heartbeatThread.Start();
+			_heartbeatTask = Task.Run(() => HeartbeatTaskLoopAsync(_cts.Token), cancellationToken);
+		}
+		return Task.CompletedTask;
+	}
+
+	/// <inheritdoc/>
+	public async Task StopAsync(CancellationToken cancellationToken)
+	{
+		if (_isRunning)
+		{
+			_isRunning = false;
+			await _cts.CancelAsync();
+			if (_heartbeatTask != null)
+			{
+				await _heartbeatTask;
+			}
+		}
+	}
+
+	private void HeartbeatThreadLoop()
+	{
+		try
+		{
+			while (_isRunning)
+			{
+				TimeSpan lastGcPauseDuration = GC.GetTotalPauseDuration();
+				DateTime lastThreadHeartbeat = DateTime.UtcNow;
+			
+				Thread.Sleep(_heartbeatThreadInterval);
+				if (ArtificialThreadDelay != null)
+				{
+					Thread.Sleep(ArtificialThreadDelay.Value);
+				}
+			
+				TimeSpan heartbeatThreadDuration = DateTime.UtcNow - lastThreadHeartbeat;
+				if (heartbeatThreadDuration > _maxThreadStallDuration)
+				{
+					_logger.LogWarning("Thread loop stalled for {StallTime} ms", (int)heartbeatThreadDuration.TotalMilliseconds);
+				}
+			
+				long ticks = Interlocked.Read(ref _lastTaskHeartbeatTicks);
+				TimeSpan heartbeatTaskDuration = DateTime.UtcNow - new DateTime(ticks);
+				if (heartbeatTaskDuration > _maxTaskStallDuration)
+				{
+					_logger.LogWarning("Task loop stalled for {StallTime} ms", (int)heartbeatTaskDuration.TotalMilliseconds);
+				}
+			
+				TimeSpan gcPauseDelta = GC.GetTotalPauseDuration() - lastGcPauseDuration;
+				if (gcPauseDelta > _maxGcPauseDelta)
+				{
+					_logger.LogWarning("GC paused process for {PauseDuration} ms over time window {TimeWindow}", (int)gcPauseDelta.TotalMilliseconds, (int)heartbeatThreadDuration.TotalMilliseconds);
+				}
+			}
+		}
+		catch (Exception e)
+		{
+			_logger.LogError(e, "Error in watchdog thread loop");
+		}
+		
+		_logger.LogInformation("Thread heartbeat loop ended");
+	}
+
+	private async Task HeartbeatTaskLoopAsync(CancellationToken cancellationToken)
+	{
+		try
+		{
+			while (!cancellationToken.IsCancellationRequested)
+			{
+				Interlocked.Exchange(ref _lastTaskHeartbeatTicks, DateTime.UtcNow.Ticks);
+				await Task.Delay(_heartbeatTaskInterval, cancellationToken);
+				if (ArtificialTaskDelay != null)
+				{
+					Thread.Sleep(ArtificialTaskDelay.Value);
+				}
+			}
+		}
+		catch (OperationCanceledException)
+		{
+			// Ignore
+		}
+		
+		_logger.LogInformation("Task heartbeat loop ended");
+	}
+}
\ No newline at end of file
