Fix test health notification replacing thread first message when trying to post a thread message

#rnx
#Synced-CL 49205937

diff --git a/Engine/Source/Programs/Horde/Plugins/Build/HordeServer.Build/Notifications/Sinks/SlackNotificationSink.cs b/Engine/Source/Programs/Horde/Plugins/Build/HordeServer.Build/Notifications/Sinks/SlackNotificationSink.cs
--- a/Engine/Source/Programs/Horde/Plugins/Build/HordeServer.Build/Notifications/Sinks/SlackNotificationSink.cs
+++ b/Engine/Source/Programs/Horde/Plugins/Build/HordeServer.Build/Notifications/Sinks/SlackNotificationSink.cs
@@ -2754,7 +2754,7 @@ namespace HordeServer.Notifications.Sinks
 					// mark fixed
 					SlackMessage fixedMessage = new SlackMessage();
 					fixedMessage.Text = $"Test Health was restored to {report.State}.";
-					await SendOrUpdateMessageToThreadAsync(recipient, eventId, null, threadId, fixedMessage, cancellationToken);
+					await _slackClient.PostMessageToThreadAsync(threadId, fixedMessage, cancellationToken);
 
 					// delete message state
 					await DeleteMessageStateAsync(recipient, eventId, cancellationToken);
@@ -2766,7 +2766,7 @@ namespace HordeServer.Notifications.Sinks
 					SlackMessage pingMessage = new SlackMessage();
 					pingMessage.Text = report.State == report.PreviousState ?
 							$"No health improvement has been recorded for {daysSinceLastUpdate} days. Consider disabling the Test or set it under Audit." : $"Test health has changed from {report.PreviousState} to {report.State}.";
-					await SendOrUpdateMessageToThreadAsync(recipient, eventId, null, threadId, pingMessage, cancellationToken);
+					await _slackClient.PostMessageToThreadAsync(threadId, pingMessage, cancellationToken);
 				}
 			}
 		}
