Horde: Move ISystemMetrics to shared

#Synced-CL 0

diff --git a/Engine/Source/Programs/Horde/HordeAgent/AgentApp.cs b/Engine/Source/Programs/Horde/HordeAgent/AgentApp.cs
--- a/Engine/Source/Programs/Horde/HordeAgent/AgentApp.cs
+++ b/Engine/Source/Programs/Horde/HordeAgent/AgentApp.cs
@@ -9,6 +9,7 @@ using System.Text.Json;
 using System.Text.Json.Serialization;
 using EpicGames.Core;
 using EpicGames.Horde;
+using EpicGames.Horde.Agents.Metrics;
 using HordeAgent.Leases;
 using HordeAgent.Leases.Handlers;
 using HordeAgent.Services;
diff --git a/Engine/Source/Programs/Horde/HordeAgent/Leases/LeaseHandler.cs b/Engine/Source/Programs/Horde/HordeAgent/Leases/LeaseHandler.cs
--- a/Engine/Source/Programs/Horde/HordeAgent/Leases/LeaseHandler.cs
+++ b/Engine/Source/Programs/Horde/HordeAgent/Leases/LeaseHandler.cs
@@ -3,6 +3,7 @@
 using System.Diagnostics;
 using EpicGames.Core;
 using EpicGames.Horde.Agents.Leases;
+using EpicGames.Horde.Agents.Metrics;
 using Google.Protobuf;
 using Google.Protobuf.WellKnownTypes;
 using HordeAgent.Services;
diff --git a/Engine/Source/Programs/Horde/HordeAgent/Leases/LeaseManager.cs b/Engine/Source/Programs/Horde/HordeAgent/Leases/LeaseManager.cs
--- a/Engine/Source/Programs/Horde/HordeAgent/Leases/LeaseManager.cs
+++ b/Engine/Source/Programs/Horde/HordeAgent/Leases/LeaseManager.cs
@@ -4,6 +4,7 @@ using System.Diagnostics;
 using System.Text;
 using EpicGames.Core;
 using EpicGames.Horde.Agents.Leases;
+using EpicGames.Horde.Agents.Metrics;
 using Google.Protobuf.WellKnownTypes;
 using Grpc.Core;
 using HordeAgent.Services;
diff --git a/Engine/Source/Programs/Horde/HordeAgent/Services/TelemetryService.cs b/Engine/Source/Programs/Horde/HordeAgent/Services/TelemetryService.cs
--- a/Engine/Source/Programs/Horde/HordeAgent/Services/TelemetryService.cs
+++ b/Engine/Source/Programs/Horde/HordeAgent/Services/TelemetryService.cs
@@ -6,158 +6,13 @@ using System.Runtime.InteropServices;
 using System.Runtime.Versioning;
 using System.Text;
 using EpicGames.Core;
+using EpicGames.Horde.Agents.Metrics;
 using Microsoft.Extensions.Hosting;
 using Microsoft.Extensions.Logging;
 using static HordeAgent.Services.WindowsCapabilities;
 
 namespace HordeAgent.Services;
 
-/// <summary>
-/// Metrics for CPU usage
-/// </summary>
-public class CpuMetrics
-{
-	/// <summary>
-	/// Percentage of time the CPU was busy executing code in user space
-	/// </summary>
-	public float User { get; set; }
-
-	/// <summary>
-	/// Percentage of time the CPU was busy executing code in kernel space
-	/// </summary>
-	public float System { get; set; }
-
-	/// <summary>
-	/// Percentage of time the CPU was idling
-	/// </summary>
-	public float Idle { get; set; }
-
-	/// <inheritdoc />
-	public override string ToString()
-	{
-		return $"User={User,5:F1}% System={System,5:F1}% Idle={Idle,5:F1}%";
-	}
-}
-
-/// <summary>
-/// Metrics for memory usage
-/// </summary>
-public class MemoryMetrics
-{
-	/// <summary>
-	/// Total memory installed (kibibytes)
-	/// </summary>
-	public uint Total { get; set; }
-
-	/// <summary>
-	/// Available memory (kibibytes)
-	/// </summary>
-	public uint Available { get; set; }
-
-	/// <summary>
-	/// Used memory (kibibytes)
-	/// </summary>
-	public uint Used { get; set; }
-
-	/// <summary>
-	/// Used memory (percentage)
-	/// </summary>
-	public float UsedPercentage { get; set; }
-
-	/// <inheritdoc />
-	public override string ToString()
-	{
-		return $"Total={Total} kB, Available={Available} kB, Used={Used} kB, UsedPct={UsedPercentage * 100.0:F1} %";
-	}
-}
-
-/// <summary>
-/// Metrics for free disk space
-/// </summary>
-public class DiskMetrics
-{
-	/// <summary>
-	/// Amount of free space on the drive
-	/// </summary>
-	public long FreeSpace { get; set; }
-
-	/// <summary>
-	/// Total size of the drive
-	/// </summary>
-	public long TotalSize { get; set; }
-	
-	/// <inheritdoc />
-	public override string ToString()
-	{
-		return $"Total={TotalSize} kB, Free={FreeSpace}";
-	}
-}
-
-/// <summary>
-/// Metrics for CPU time spent
-/// </summary>
-public class CpuTimeMetrics
-{
-	/// <summary>
-	/// CPU time spent in user mode
-	/// </summary>
-	public TimeSpan User { get; set; }
-
-	/// <summary>
-	/// CPU time spent in system/kernel mode
-	/// </summary>
-	public TimeSpan System { get; set; }
-
-	/// <summary>
-	/// CPU time spent idle
-	/// </summary>
-	public TimeSpan Idle { get; set; }
-
-	/// <summary>
-	/// Subtract times from other object passed in
-	/// </summary>
-	public CpuTimeMetrics Subtract(CpuTimeMetrics other)
-	{
-		return new CpuTimeMetrics { User = User - other.User, System = System - other.System, Idle = Idle - other.Idle };
-	}
-	
-	/// <inheritdoc />
-	public override string ToString()
-	{
-		return $"User={(int)User.TotalMilliseconds} System={(int)System.TotalMilliseconds} Idle={(int)Idle.TotalMilliseconds}";
-	}
-}
-
-/// <summary>
-/// OS agnostic interface for retrieving system metrics (CPU, memory etc)
-/// </summary>
-public interface ISystemMetrics
-{
-	/// <summary>
-	/// Get CPU usage metrics
-	/// </summary>
-	/// <returns>An object with CPU usage metrics</returns>
-	CpuMetrics? GetCpu();
-
-	/// <summary>
-	/// Get cumulative CPU time spent since system start
-	/// </summary>
-	/// <returns>An object with CPU time metrics</returns>
-	CpuTimeMetrics? GetCpuTime();
-	
-	/// <summary>
-	/// Get memory usage metrics
-	/// </summary>
-	/// <returns>An object with memory usage metrics</returns>
-	MemoryMetrics? GetMemory();
-
-	/// <summary>
-	/// Gets HDD metrics
-	/// </summary>
-	/// <returns>An object with disk usage metrics</returns>
-	DiskMetrics? GetDisk();
-}
-
 /// <summary>
 /// Default implementation of <see cref="ISystemMetrics"/>
 /// </summary>
diff --git a/Engine/Source/Programs/Shared/EpicGames.Horde/Agents/Metrics/Metrics.cs b/Engine/Source/Programs/Shared/EpicGames.Horde/Agents/Metrics/Metrics.cs
new file mode 100644
--- /dev/null
+++ b/Engine/Source/Programs/Shared/EpicGames.Horde/Agents/Metrics/Metrics.cs
@@ -0,0 +1,151 @@
+// Copyright Epic Games, Inc. All Rights Reserved.
+
+using System;
+
+namespace EpicGames.Horde.Agents.Metrics;
+
+/// <summary>
+/// Metrics for CPU usage
+/// </summary>
+public class CpuMetrics
+{
+	/// <summary>
+	/// Percentage of time the CPU was busy executing code in user space
+	/// </summary>
+	public float User { get; set; }
+
+	/// <summary>
+	/// Percentage of time the CPU was busy executing code in kernel space
+	/// </summary>
+	public float System { get; set; }
+
+	/// <summary>
+	/// Percentage of time the CPU was idling
+	/// </summary>
+	public float Idle { get; set; }
+
+	/// <inheritdoc />
+	public override string ToString()
+	{
+		return $"User={User,5:F1}% System={System,5:F1}% Idle={Idle,5:F1}%";
+	}
+}
+
+/// <summary>
+/// Metrics for memory usage
+/// </summary>
+public class MemoryMetrics
+{
+	/// <summary>
+	/// Total memory installed (kibibytes)
+	/// </summary>
+	public uint Total { get; set; }
+
+	/// <summary>
+	/// Available memory (kibibytes)
+	/// </summary>
+	public uint Available { get; set; }
+
+	/// <summary>
+	/// Used memory (kibibytes)
+	/// </summary>
+	public uint Used { get; set; }
+
+	/// <summary>
+	/// Used memory (percentage)
+	/// </summary>
+	public float UsedPercentage { get; set; }
+
+	/// <inheritdoc />
+	public override string ToString()
+	{
+		return $"Total={Total} kB, Available={Available} kB, Used={Used} kB, UsedPct={UsedPercentage * 100.0:F1} %";
+	}
+}
+
+/// <summary>
+/// Metrics for free disk space
+/// </summary>
+public class DiskMetrics
+{
+	/// <summary>
+	/// Amount of free space on the drive
+	/// </summary>
+	public long FreeSpace { get; set; }
+
+	/// <summary>
+	/// Total size of the drive
+	/// </summary>
+	public long TotalSize { get; set; }
+
+	/// <inheritdoc />
+	public override string ToString()
+	{
+		return $"Total={TotalSize} kB, Free={FreeSpace}";
+	}
+}
+
+/// <summary>
+/// Metrics for CPU time spent
+/// </summary>
+public class CpuTimeMetrics
+{
+	/// <summary>
+	/// CPU time spent in user mode
+	/// </summary>
+	public TimeSpan User { get; set; }
+
+	/// <summary>
+	/// CPU time spent in system/kernel mode
+	/// </summary>
+	public TimeSpan System { get; set; }
+
+	/// <summary>
+	/// CPU time spent idle
+	/// </summary>
+	public TimeSpan Idle { get; set; }
+
+	/// <summary>
+	/// Subtract times from other object passed in
+	/// </summary>
+	public CpuTimeMetrics Subtract(CpuTimeMetrics other)
+	{
+		return new CpuTimeMetrics { User = User - other.User, System = System - other.System, Idle = Idle - other.Idle };
+	}
+
+	/// <inheritdoc />
+	public override string ToString()
+	{
+		return $"User={(int)User.TotalMilliseconds} System={(int)System.TotalMilliseconds} Idle={(int)Idle.TotalMilliseconds}";
+	}
+}
+
+/// <summary>
+/// OS-agnostic interface for retrieving system metrics (CPU, memory etc)
+/// </summary>
+public interface ISystemMetrics
+{
+	/// <summary>
+	/// Get CPU usage metrics
+	/// </summary>
+	/// <returns>An object with CPU usage metrics</returns>
+	CpuMetrics? GetCpu();
+
+	/// <summary>
+	/// Get cumulative CPU time spent since system start
+	/// </summary>
+	/// <returns>An object with CPU time metrics</returns>
+	CpuTimeMetrics? GetCpuTime();
+
+	/// <summary>
+	/// Get memory usage metrics
+	/// </summary>
+	/// <returns>An object with memory usage metrics</returns>
+	MemoryMetrics? GetMemory();
+
+	/// <summary>
+	/// Gets HDD metrics
+	/// </summary>
+	/// <returns>An object with disk usage metrics</returns>
+	DiskMetrics? GetDisk();
+}
