Horde: Another take on duplicate json Telemetry key exceptions

diff --git a/Engine/Source/Programs/Horde/Plugins/Analytics/HordeServer.Analytics/Telemetry/TelemetryController.cs b/Engine/Source/Programs/Horde/Plugins/Analytics/HordeServer.Analytics/Telemetry/TelemetryController.cs
--- a/Engine/Source/Programs/Horde/Plugins/Analytics/HordeServer.Analytics/Telemetry/TelemetryController.cs
+++ b/Engine/Source/Programs/Horde/Plugins/Analytics/HordeServer.Analytics/Telemetry/TelemetryController.cs
@@ -1,14 +1,13 @@
 // Copyright Epic Games, Inc. All Rights Reserved.
 
 using System.ComponentModel.DataAnnotations;
-using System.Text.Json.Nodes;
+using System.Text.Json;
 using EpicGames.Horde.Telemetry;
 using EpicGames.Horde.Telemetry.Metrics;
 using HordeServer.Telemetry.Metrics;
 using HordeServer.Utilities;
 using Microsoft.AspNetCore.Authorization;
 using Microsoft.AspNetCore.Mvc;
-using Microsoft.AspNetCore.Mvc.ModelBinding.Validation;
 using Microsoft.Extensions.Options;
 
 namespace HordeServer.Telemetry
@@ -179,7 +178,7 @@ namespace HordeServer.Telemetry
 		/// <param name="uploadType">Type of data being uploaded</param>
 		[HttpPost]
 		[Route("/api/v1/telemetry/{storeId}")]
-		public ActionResult PostEvent(TelemetryStoreId storeId, [FromBody][ValidateNever] PostTelemetryEventStreamRequest request, [FromQuery][Required] string appId, [FromQuery][Required] string appVersion, [FromQuery][Required] string appEnvironment, [FromQuery][Required] TelemetryUploadType uploadType)
+		public ActionResult PostEvent(TelemetryStoreId storeId, [FromBody] PostTelemetryEventStreamRequest request, [FromQuery][Required] string appId, [FromQuery][Required] string appVersion, [FromQuery][Required] string appEnvironment, [FromQuery][Required] TelemetryUploadType uploadType)
 		{
 			if (uploadType != TelemetryUploadType.EtEventStream)
 			{
@@ -187,7 +186,7 @@ namespace HordeServer.Telemetry
 			}
 
 			TelemetryRecordMeta recordMeta = new TelemetryRecordMeta(AppId: appId, AppVersion: appVersion, AppEnvironment: appEnvironment);
-			foreach (JsonObject eventPayload in request.Events)
+			foreach (JsonElement eventPayload in request.Events)
 			{
 				_telemetryManager.WriteEvent(storeId, recordMeta, eventPayload);
 			}
diff --git a/Engine/Source/Programs/Shared/EpicGames.Horde/HordeHttpClient.cs b/Engine/Source/Programs/Shared/EpicGames.Horde/HordeHttpClient.cs
--- a/Engine/Source/Programs/Shared/EpicGames.Horde/HordeHttpClient.cs
+++ b/Engine/Source/Programs/Shared/EpicGames.Horde/HordeHttpClient.cs
@@ -9,7 +9,6 @@ using System.Net.Http;
 using System.Net.Http.Headers;
 using System.Net.Http.Json;
 using System.Text.Json;
-using System.Text.Json.Nodes;
 using System.Text.RegularExpressions;
 using System.Threading;
 using System.Threading.Tasks;
@@ -454,11 +453,11 @@ namespace EpicGames.Horde
 				for (int i = 0; i < payloads.Length; ++i)
 				{
 					object payload = payloads[i];
-					JsonNode? node = JsonSerializer.SerializeToNode(payload, JsonSerializerOptions);
+					JsonElement element = JsonSerializer.SerializeToElement(payload, JsonSerializerOptions);
 
-					if (node != null && node is JsonObject obj)
+					if (element.ValueKind == JsonValueKind.Object)
 					{
-						request.Events.Add(obj);
+						request.Events.Add(element);
 					}
 					else
 					{
diff --git a/Engine/Source/Programs/Shared/EpicGames.Horde/Telemetry/TelemetryMessages.cs b/Engine/Source/Programs/Shared/EpicGames.Horde/Telemetry/TelemetryMessages.cs
--- a/Engine/Source/Programs/Shared/EpicGames.Horde/Telemetry/TelemetryMessages.cs
+++ b/Engine/Source/Programs/Shared/EpicGames.Horde/Telemetry/TelemetryMessages.cs
@@ -20,7 +20,7 @@ namespace EpicGames.Horde.Telemetry
 		/// <summary>
 		/// List of telemetry events
 		/// </summary>
-		public List<System.Text.Json.Nodes.JsonObject> Events { get; set; } = new List<System.Text.Json.Nodes.JsonObject>();
+		public List<System.Text.Json.JsonElement> Events { get; set; } = new List<System.Text.Json.JsonElement>();
 	}
 
 	/// <summary>
