Horde: Fix issue with agent tooltip on job timeline
#github https://github.com/EpicGames/UnrealEngine/pull/14146
#submitool bypass

diff --git a/Engine/Source/Programs/Horde/HordeDashboard/src/components/jobDetailsV2/JobDetailTimeline.tsx b/Engine/Source/Programs/Horde/HordeDashboard/src/components/jobDetailsV2/JobDetailTimeline.tsx
--- a/Engine/Source/Programs/Horde/HordeDashboard/src/components/jobDetailsV2/JobDetailTimeline.tsx
+++ b/Engine/Source/Programs/Horde/HordeDashboard/src/components/jobDetailsV2/JobDetailTimeline.tsx
@@ -786,7 +786,7 @@ const GraphTooltip: React.FC<{ dataView: TimelineDataView }> = observer(({ dataV
       }
 
       if (agent) {
-         valueElement = <button className="horde-link" style={{ borderWidth: 0, padding: 0, backgroundColor: modeColors.background, cursor: "pointer" }}>
+         valueElement = <button className="horde-link" style={{ borderWidth: 0, padding: 0, backgroundColor: modeColors.background, cursor: "pointer", color: modeColors.text }}>
             {valueElement}
          </button>
       }
