[Horde][BuildHealth][Server] Description field & admin powers for Build health filter, as well as returning full owner object for better messaging around follow-ups.

#Synced-CL 49331470

diff --git a/Engine/Source/Programs/Horde/Plugins/Build/HordeServer.Build/BuildHealth/BuildHealthFilterCollection.cs b/Engine/Source/Programs/Horde/Plugins/Build/HordeServer.Build/BuildHealth/BuildHealthFilterCollection.cs
--- a/Engine/Source/Programs/Horde/Plugins/Build/HordeServer.Build/BuildHealth/BuildHealthFilterCollection.cs
+++ b/Engine/Source/Programs/Horde/Plugins/Build/HordeServer.Build/BuildHealth/BuildHealthFilterCollection.cs
@@ -50,7 +50,7 @@ namespace HordeServer.Build.BuildHealth
 				return null;
 			}
 
-			BuildHealthFilterDocument newBuildHealthFilter = new BuildHealthFilterDocument() { Id = new BuildHealthFilterId(BinaryIdUtils.CreateNew()), Owner = owner, FilterName = filterRequest.FilterName, FilterProject = filterRequest.FilterProject, FilterQuery = filterRequest.FilterQuery, UpdateTimeUtc = DateTime.UtcNow };
+			BuildHealthFilterDocument newBuildHealthFilter = new BuildHealthFilterDocument() { Id = new BuildHealthFilterId(BinaryIdUtils.CreateNew()), Owner = owner, FilterName = filterRequest.FilterName, FilterDescription = filterRequest.FilterDescription, FilterProject = filterRequest.FilterProject, FilterQuery = filterRequest.FilterQuery, UpdateTimeUtc = DateTime.UtcNow };
 
 			await _buildHealthFilterDocuments.InsertOneAsync(newBuildHealthFilter, null, cancellationToken);
 			return newBuildHealthFilter;
@@ -71,6 +71,14 @@ namespace HordeServer.Build.BuildHealth
 			{
 				updates.Add(Builders<BuildHealthFilterDocument>.Update.Set(x => x.FilterName, filterRequest.FilterName));
 			}
+
+			// Check for an explicit null, as Empty/Whitespace is a valid erasure.
+			if (filterRequest.FilterDescription != null)
+			{
+				// If we have passed a whitespace, we will set it to null in the document.
+				updates.Add(Builders<BuildHealthFilterDocument>.Update.Set(x => x.FilterDescription, String.IsNullOrWhiteSpace(filterRequest.FilterDescription) ? null : filterRequest.FilterDescription));
+			}
+
 			if (!String.IsNullOrEmpty(filterRequest.FilterQuery))
 			{
 				updates.Add
@@ -145,6 +153,10 @@ namespace HordeServer.Build.BuildHealth
 			[BsonElement("fn")]
 			public string FilterName { get; set; } = default!;
 
+			/// <inheritdoc/>
+			[BsonElement("fd")]
+			public string? FilterDescription { get; set; }
+
 			/// <inheritdoc/>
 			[BsonElement("fp")]
 			public ProjectId FilterProject { get; set; }
diff --git a/Engine/Source/Programs/Horde/Plugins/Build/HordeServer.Build/BuildHealth/BuildHealthFilterController.cs b/Engine/Source/Programs/Horde/Plugins/Build/HordeServer.Build/BuildHealth/BuildHealthFilterController.cs
--- a/Engine/Source/Programs/Horde/Plugins/Build/HordeServer.Build/BuildHealth/BuildHealthFilterController.cs
+++ b/Engine/Source/Programs/Horde/Plugins/Build/HordeServer.Build/BuildHealth/BuildHealthFilterController.cs
@@ -4,6 +4,7 @@ using EpicGames.Horde.BuildHealth;
 using EpicGames.Horde.Projects;
 using EpicGames.Horde.Users;
 using HordeServer.Projects;
+using HordeServer.Users;
 using HordeServer.Utilities;
 using Microsoft.AspNetCore.Authorization;
 using Microsoft.AspNetCore.Http;
@@ -24,6 +25,7 @@ namespace HordeServer.Build.BuildHealth
 		#region -- Private Members --
 
 		readonly IBuildHealthFilterCollection _buildHealthFilterCollection;
+		readonly IUserCollection _userCollection;
 		private readonly IOptionsSnapshot<BuildConfig> _buildConfig;
 
 		#endregion -- Private Members --
@@ -34,10 +36,12 @@ namespace HordeServer.Build.BuildHealth
 		/// Constructor.
 		/// </summary>
 		/// <param name="buildHealthFilterCollection">The underyling build health filter collection.</param>
+		/// <param name="userCollection">The underlying user collection used for collating ownership details.</param>
 		/// <param name="buildConfig">Build config used in accessing relevant ACL configurations.</param>
-		public BuildHealthFilterController(IBuildHealthFilterCollection buildHealthFilterCollection, IOptionsSnapshot<BuildConfig> buildConfig)
+		public BuildHealthFilterController(IBuildHealthFilterCollection buildHealthFilterCollection, IUserCollection userCollection, IOptionsSnapshot<BuildConfig> buildConfig)
 		{
 			_buildHealthFilterCollection = buildHealthFilterCollection;
+			_userCollection = userCollection;
 			_buildConfig = buildConfig;
 		}
 
@@ -70,7 +74,7 @@ namespace HordeServer.Build.BuildHealth
 
 			if (!filterUpdateRequest.IsValidAddRequest())
 			{
-				return BadRequest($"UpdateFilterRequest was not valid. Check filter name length (Max:{BuildHealthFilterRestrictions.MaxFilterNameLength}) & filter query length (Max:{BuildHealthFilterRestrictions.MaxFilterQueryLength}).");
+				return BadRequest($"UpdateFilterRequest was not valid. Check filter name length (Max:{BuildHealthFilterRestrictions.MaxFilterNameLength}) & filter query length (Max:{BuildHealthFilterRestrictions.MaxFilterQueryLength}) & filter description length (Max:{BuildHealthFilterRestrictions.MaxFilterDescripitonLength}).");
 			}
 
 			// Validate user has access for add operation.
@@ -95,7 +99,7 @@ namespace HordeServer.Build.BuildHealth
 
 			}
 
-			return new BuildHealthFilterResponse(filter);
+			return await CreateBuildHealthFilterResponseWithOwnerAsync(filter, cancellationToken);
 		}
 
 		/// <summary>
@@ -105,7 +109,7 @@ namespace HordeServer.Build.BuildHealth
 		/// <param name="filterUpdateRequest">The update request.</param>
 		/// <param name="cancellationToken">The cancellation token to use throughout the operation.</param>
 		/// <returns>The updated filter if applicable.</returns>
-		/// <remarks>Users can only update filters that don't have an owner, or ones they are an owner of.</remarks>
+		/// <remarks>Users can only update filters that don't have an owner, ones they are an owner of, or any filter if they are an admin.</remarks>
 		[HttpPut]
 		[Route("/api/v1/buildhealth/filters/{filterId}")]
 		[ProducesResponseType(typeof(BuildHealthFilterResponse), 200)]
@@ -119,7 +123,7 @@ namespace HordeServer.Build.BuildHealth
 
 			if (!filterUpdateRequest.IsValidUpdateRequest())
 			{
-				return BadRequest($"UpdateFilterRequest was not valid. Check filter name length (Max:{BuildHealthFilterRestrictions.MaxFilterNameLength}) & filter query length (Max:{BuildHealthFilterRestrictions.MaxFilterQueryLength}).");
+				return BadRequest($"UpdateFilterRequest was not valid. Check filter name length (Max:{BuildHealthFilterRestrictions.MaxFilterNameLength}) & filter query length (Max:{BuildHealthFilterRestrictions.MaxFilterQueryLength}) & filter description length (Max:{BuildHealthFilterRestrictions.MaxFilterDescripitonLength}).");
 			}
 
 			// Validate user has access for update operation.
@@ -144,7 +148,9 @@ namespace HordeServer.Build.BuildHealth
 
 			// Verify the user is able to edit this filter.
 			UserId? requestingUser = User.GetUserId();
-			if (existingFilter.Owner != null && existingFilter.Owner != requestingUser)
+			bool canEdit = existingFilter.Owner == requestingUser || User.HasAdminClaim();
+
+			if (!canEdit)
 			{
 				return Forbid($"User is not the owner of BuildHealthFilter Id: {filterId}.");
 			}
@@ -157,7 +163,7 @@ namespace HordeServer.Build.BuildHealth
 
 			}
 
-			return new BuildHealthFilterResponse(filter);
+			return await CreateBuildHealthFilterResponseWithOwnerAsync(filter, cancellationToken);
 		}
 
 		/// <summary>
@@ -184,7 +190,7 @@ namespace HordeServer.Build.BuildHealth
 				return actionResult;
 			}
 
-			return new BuildHealthFilterResponse(filter);
+			return await CreateBuildHealthFilterResponseWithOwnerAsync(filter, cancellationToken);
 		}
 
 		/// <summary>
@@ -210,7 +216,8 @@ namespace HordeServer.Build.BuildHealth
 
 			foreach (IBuildHealthFilter filter in filters)
 			{
-				returnObjects.Add(new BuildHealthFilterResponse(filter));
+				BuildHealthFilterResponse buildHealthFilter = await CreateBuildHealthFilterResponseWithOwnerAsync(filter, cancellationToken);
+				returnObjects.Add(buildHealthFilter);
 			}
 
 			return returnObjects;
@@ -222,6 +229,7 @@ namespace HordeServer.Build.BuildHealth
 		/// <param name="filterId">The id of the filter to delete.</param>
 		/// <param name="cancellationToken">The cancellation token to use throughout the operation.</param>
 		/// <returns></returns>
+		/// <remarks>Users can only delete filters that don't have an owner, ones they are an owner of, or any filter if they are an admin.</remarks>
 		[HttpDelete]
 		[Route("/api/v1/buildhealth/filters/{filterId}")]
 		public async Task<ActionResult> DeleteBuildHealthFiltersAsync(BuildHealthFilterId filterId, CancellationToken cancellationToken = default)
@@ -249,7 +257,9 @@ namespace HordeServer.Build.BuildHealth
 
 			// Verify user is able to delete this filter.
 			UserId? requestingUser = User.GetUserId();
-			if (existingFilter.Owner != null && existingFilter.Owner != requestingUser)
+			bool canEdit = existingFilter.Owner == requestingUser || User.HasAdminClaim();
+
+			if (!canEdit)
 			{
 				return Forbid("User is not the owner of BuildHealthFilter Id: {BuildHealthFilterId}.", existingFilter.Id);
 			}
@@ -266,6 +276,21 @@ namespace HordeServer.Build.BuildHealth
 
 		#region -- Private API --
 
+		private async Task<BuildHealthFilterResponse> CreateBuildHealthFilterResponseWithOwnerAsync(IBuildHealthFilter filter, CancellationToken cancellationToken)
+		{
+			if (filter.Owner != null)
+			{
+				GetThinUserInfoResponse? user = (await _userCollection.GetCachedUserAsync(filter.Owner.Value, cancellationToken))?.ToThinApiResponse();
+
+				if (user != null)
+				{
+					return new BuildHealthFilterResponse(filter, user);
+				}
+			}
+
+			return new BuildHealthFilterResponse(filter);
+		}
+
 		/// <summary>
 		/// Helper method to verify and authorize project.
 		/// </summary>
diff --git a/Engine/Source/Programs/Shared/EpicGames.Horde/BuildHealth/BuildHealthFilterMessages.cs b/Engine/Source/Programs/Shared/EpicGames.Horde/BuildHealth/BuildHealthFilterMessages.cs
--- a/Engine/Source/Programs/Shared/EpicGames.Horde/BuildHealth/BuildHealthFilterMessages.cs
+++ b/Engine/Source/Programs/Shared/EpicGames.Horde/BuildHealth/BuildHealthFilterMessages.cs
@@ -16,6 +16,11 @@ namespace EpicGames.Horde.BuildHealth
 		/// </summary>
 		public const int MaxFilterNameLength = 250;
 
+		/// <summary>
+		/// The maximum filter description length.
+		/// </summary>
+		public const int MaxFilterDescripitonLength = 500;
+
 		/// <summary>
 		/// The maximum filter query length.
 		/// </summary>
@@ -37,6 +42,11 @@ namespace EpicGames.Horde.BuildHealth
 		/// </summary>
 		public string FilterName { get; set; } = default!;
 
+		/// <summary>
+		/// The description of the filter.
+		/// </summary>
+		public string? FilterDescription { get; init; }
+
 		/// <summary>
 		/// The filter query.
 		/// </summary>
@@ -48,6 +58,12 @@ namespace EpicGames.Horde.BuildHealth
 		/// <returns>True if valid, false otherwise.</returns>
 		public bool IsValidFilterName() => (FilterName != null && FilterName.Length <= BuildHealthFilterRestrictions.MaxFilterNameLength);
 
+		/// <summary>
+		/// Returns whether the filter description request is valid or not.
+		/// </summary>
+		/// <returns>True if valid, false otherwise.</returns>
+		public bool IsValidFilterDescription() => (FilterDescription == null || FilterDescription.Length <= BuildHealthFilterRestrictions.MaxFilterDescripitonLength);
+
 		/// <summary>
 		/// Returns whether the filter query request is valid or not.
 		/// </summary>
@@ -58,7 +74,7 @@ namespace EpicGames.Horde.BuildHealth
 		/// Returns whether the filter add request is valid.
 		/// </summary>
 		/// <returns>True if valid, false otherwise.</returns>
-		public bool IsValidAddRequest() => IsValidFilterQuery() && IsValidFilterName();
+		public bool IsValidAddRequest() => IsValidFilterQuery() && IsValidFilterName() && IsValidFilterDescription();
 	}
 
 	/// <summary>
@@ -71,6 +87,11 @@ namespace EpicGames.Horde.BuildHealth
 		/// </summary>
 		public string? FilterName { get; set; } = default!;
 
+		/// <summary>
+		/// The description of the filter.
+		/// </summary>
+		public string? FilterDescription { get; init; }
+
 		/// <summary>
 		/// The filter query.
 		/// </summary>
@@ -82,6 +103,12 @@ namespace EpicGames.Horde.BuildHealth
 		/// <returns>True if valid, false otherwise.</returns>
 		public bool IsValidFilterName() => (FilterName == null || FilterName.Length <= BuildHealthFilterRestrictions.MaxFilterNameLength);
 
+		/// <summary>
+		/// Returns whether the filter description request is valid or not.
+		/// </summary>
+		/// <returns>True if valid, false otherwise.</returns>
+		public bool IsValidFilterDescription() => (FilterDescription == null || FilterDescription.Length <= BuildHealthFilterRestrictions.MaxFilterDescripitonLength);
+
 		/// <summary>
 		/// Returns whether the filter query request is valid or not.
 		/// </summary>
@@ -92,7 +119,7 @@ namespace EpicGames.Horde.BuildHealth
 		/// Returns whether the filter update request is valid.
 		/// </summary>
 		/// <returns>True if valid, false otherwise.</returns>
-		public bool IsValidUpdateRequest() => IsValidFilterQuery() && IsValidFilterName();
+		public bool IsValidUpdateRequest() => IsValidFilterQuery() && IsValidFilterName() && IsValidFilterDescription();
 	}
 
 	/// <summary>
@@ -108,7 +135,7 @@ namespace EpicGames.Horde.BuildHealth
 		/// <summary>
 		/// The owner of the filter.
 		/// </summary>
-		public UserId? Owner { get; init; } = default;
+		public GetThinUserInfoResponse? Owner { get; init; } = default;
 
 		/// <summary>
 		/// The project the filter belongs to.
@@ -120,6 +147,11 @@ namespace EpicGames.Horde.BuildHealth
 		/// </summary>
 		public string FilterName { get; init; }
 
+		/// <summary>
+		/// The description of the filter.
+		/// </summary>
+		public string? FilterDescription { get; init; }
+
 		/// <summary>
 		/// The filter query.
 		/// </summary>
@@ -138,9 +170,26 @@ namespace EpicGames.Horde.BuildHealth
 		public BuildHealthFilterResponse(IBuildHealthFilter filter)
 		{
 			Id = filter.Id;
-			Owner = filter.Owner;
+			Owner = null;
+			FilterName = filter.FilterName;
 			FilterProject = filter.FilterProject;
+			FilterDescription = filter.FilterDescription;
+			FilterQuery = filter.FilterQuery;
+			UpdateTimeUtc = filter.UpdateTimeUtc;
+		}
+
+		/// <summary>
+		/// Creates a build health filter response from a provided build health filter, and the ownership object.
+		/// </summary>
+		/// <param name="filter">The source filter.</param>
+		/// <param name="owner">The owner of the filter.</param>
+		public BuildHealthFilterResponse(IBuildHealthFilter filter, GetThinUserInfoResponse owner)
+		{
+			Id = filter.Id;
+			Owner = owner;
 			FilterName = filter.FilterName;
+			FilterProject = filter.FilterProject;
+			FilterDescription = filter.FilterDescription;
 			FilterQuery = filter.FilterQuery;
 			UpdateTimeUtc = filter.UpdateTimeUtc;
 		}
diff --git a/Engine/Source/Programs/Shared/EpicGames.Horde/BuildHealth/IBuildHealthFilter.cs b/Engine/Source/Programs/Shared/EpicGames.Horde/BuildHealth/IBuildHealthFilter.cs
--- a/Engine/Source/Programs/Shared/EpicGames.Horde/BuildHealth/IBuildHealthFilter.cs
+++ b/Engine/Source/Programs/Shared/EpicGames.Horde/BuildHealth/IBuildHealthFilter.cs
@@ -31,6 +31,11 @@ namespace EpicGames.Horde.BuildHealth
 		/// </summary>
 		string FilterName { get; }
 
+		/// <summary>
+		/// A description of the filter.
+		/// </summary>
+		string? FilterDescription { get; }
+
 		/// <summary>
 		/// The project the filter belongs to.
 		/// </summary>
