[Horde][BuildHealth][Dashboard] Build health filter QoL improvements, preparation for wide users, and minor bug fixes:
- Description field
- Full ownership context returned with object
- Dropdown to use seach & items ordered alphabetically
- Dropdown tooltips to include owner, proper time format, and description
- Labels more intuitive based on whether the owner has authority to CRUD build filters, or not
- Syntheized job step outcomes for when batches fail

#Synced-CL 49331470

diff --git a/Engine/Source/Programs/Horde/HordeDashboard/src/backend/Api.ts b/Engine/Source/Programs/Horde/HordeDashboard/src/backend/Api.ts
--- a/Engine/Source/Programs/Horde/HordeDashboard/src/backend/Api.ts
+++ b/Engine/Source/Programs/Horde/HordeDashboard/src/backend/Api.ts
@@ -5430,6 +5430,9 @@ export type AddBuildHealthFiltersRequest = {
 	/// The filter name.
 	filterName: string;
 
+	/// The filter description.
+	filterDescription?: string;
+
 	/// The filter query.
 	filterQuery: string;
 }
@@ -5441,6 +5444,9 @@ export type UpdateBuildHealthFiltersRequest = {
 	/// The filter name.
 	filterName: string;
 
+	/// The filter description.
+	filterDescription?: string;
+
 	/// The filter query.
 	filterQuery: string;
 }
@@ -5457,7 +5463,7 @@ export type GetBuildHealthFilterResponse = {
 	id: string;
 
 	/// The owner of the filter.
-	owner: string;
+	owner: GetThinUserInfoResponse;
 
 	/// The project the filter belongs to.
 	filterProject: string;
@@ -5465,6 +5471,9 @@ export type GetBuildHealthFilterResponse = {
 	/// The name of the filter.
 	filterName: string;
 
+	/// The filter description.
+	filterDescription?: string;
+
 	/// The query of the filter.
 	filterQuery: string;
 
diff --git a/Engine/Source/Programs/Horde/HordeDashboard/src/components/buildhealth/BuildHealthDataTypes.ts b/Engine/Source/Programs/Horde/HordeDashboard/src/components/buildhealth/BuildHealthDataTypes.ts
--- a/Engine/Source/Programs/Horde/HordeDashboard/src/components/buildhealth/BuildHealthDataTypes.ts
+++ b/Engine/Source/Programs/Horde/HordeDashboard/src/components/buildhealth/BuildHealthDataTypes.ts
@@ -1,5 +1,6 @@
 // Copyright Epic Games, Inc. All Rights Reserved.
 
+import { mergeStyleSets } from "@fluentui/react";
 import { GetLabelStateResponse, GetStepResponse, GetTemplateRefResponse } from "horde/backend/Api";
 
 export const QUERY_SCHEMA_PARAM_NAME: string = "query_schema";
@@ -25,6 +26,25 @@ export const STEPS_URL_PARAM_ALL_STEPS: string = "stepNames_AllKnownValues";
 export const STATES_URL_PARAM_ALL_STATES: string = "stateNames_AllKnownValues";
 export const JOB_START_URL_PARAM_ALL_NAMES: string = "jobStart_akv";
 
+// #region -- Styles --
+
+export const dropdownNavigationButtons = mergeStyleSets({
+    base: {
+        marginLeft: 4,
+        padding: "2px 6px",
+        fontSize: 12,
+        cursor: "pointer",
+    }
+});
+
+export const dropdownOptions = mergeStyleSets({
+    base: {
+        padding: "4px 8px",
+    }
+});
+
+// #endregion -- Styles --
+
 /**
  * Convenience type to fold in extra metadata relevant for the Template. 
  */
diff --git a/Engine/Source/Programs/Horde/HordeDashboard/src/components/buildhealth/BuildHealthFilterComponent.tsx b/Engine/Source/Programs/Horde/HordeDashboard/src/components/buildhealth/BuildHealthFilterComponent.tsx
--- a/Engine/Source/Programs/Horde/HordeDashboard/src/components/buildhealth/BuildHealthFilterComponent.tsx
+++ b/Engine/Source/Programs/Horde/HordeDashboard/src/components/buildhealth/BuildHealthFilterComponent.tsx
@@ -1,5 +1,5 @@
-import { DirectionalHint, Dropdown, IconButton, IDropdownOption, ISelectableOption, Label, Stack, TooltipHost } from "@fluentui/react";
-import { GetBuildHealthFilterResponse } from "horde/backend/Api";
+import { DirectionalHint, IconButton, IDropdownOption, ISelectableOption, Label, Stack, TooltipHost } from "@fluentui/react";
+import { GetBuildHealthFilterResponse, GetThinUserInfoResponse } from "horde/backend/Api";
 import { BuildHealthDataHandler } from "./BuildHealthDataHandler";
 import { BuildHealthOptionsController } from "./BuildHealthOptionsController";
 import { useState } from "react";
@@ -7,6 +7,8 @@ import { observer } from "mobx-react-lite";
 import { BuildHealthFilterModal, BuildHealthFilterModalProps } from "./BuildHealthFilterModal";
 import { parseBuildHealthQueryParams, upgradeBuildHealthQueryParams, loadBuildHealthOptionsFromParams, toNavigationQuery } from "./BuildHealthOptions";
 import { getHordeStyling } from "horde/styles/Styles";
+import { SearchableDropdown } from "./SearchableDropdown";
+import { getDateTimeString } from "./stepoutcome/StepOutcomeUtilities";
 
 // #region -- Helper Functions --
 
@@ -49,6 +51,9 @@ export class EditableBuildHealthFilter {
     /// The name of the filter.
     filterName: string;
 
+    /// The description of the filter.
+    filterDescription?: string;
+
     /// The query of the filter. If null, there is no prior state (new filter).
     oldFilterQuery: string;
 
@@ -59,7 +64,7 @@ export class EditableBuildHealthFilter {
     lastUpdated: Date | null;
 
     /// The filter owner. If null, there is no prior state (new filter).
-    filterOwner: string | null;
+    filterOwner: GetThinUserInfoResponse | null;
 
     /// The old query, in human readable form. If null, there is no prior state (new filter).
     oldHumanReadableQuery: string | null;
@@ -67,15 +72,16 @@ export class EditableBuildHealthFilter {
     /// The new query, in human readable form.
     newHumanReadableQuery: string;
 
-    private constructor(filterName: string, filterQuery: string, newFilterQuery: string, oldHumanReadableQuery: string | null, newHumanReadableQuery: string, filterId: string | null, filterProject: string, ownerId: string | null, lastUpdated: Date | string | null) {
+    private constructor(filterName: string, filterQuery: string, newFilterQuery: string, oldHumanReadableQuery: string | null, newHumanReadableQuery: string, filterId: string | null, filterProject: string, owner: GetThinUserInfoResponse | null, lastUpdated: Date | string | null, filterDescription?: string) {
         this.filterName = filterName;
+        this.filterDescription = filterDescription;
         this.oldFilterQuery = filterQuery;
         this.newFilterQuery = newFilterQuery;
         this.filterId = filterId;
         this.oldHumanReadableQuery = oldHumanReadableQuery;
         this.newHumanReadableQuery = newHumanReadableQuery;
         this.filterProjectId = filterProject;
-        this.filterOwner = ownerId;
+        this.filterOwner = owner;
         this.lastUpdated = lastUpdated === null ? null : new Date(lastUpdated);
     }
 
@@ -88,7 +94,7 @@ export class EditableBuildHealthFilter {
      * @returns The editable form of the existing filter.
      */
     static produceFromExisting(existingFilter: GetBuildHealthFilterResponse, newFilterQuery: string, oldHumanReadableQuery: string, newHumanReadableQuery: string): EditableBuildHealthFilter {
-        return new EditableBuildHealthFilter(existingFilter.filterName, existingFilter.filterQuery, newFilterQuery, oldHumanReadableQuery, newHumanReadableQuery, existingFilter.id, existingFilter.filterProject, existingFilter.owner, existingFilter.updateTimeUtc);
+        return new EditableBuildHealthFilter(existingFilter.filterName, existingFilter.filterQuery, newFilterQuery, oldHumanReadableQuery, newHumanReadableQuery, existingFilter.id, existingFilter.filterProject, existingFilter.owner, existingFilter.updateTimeUtc, existingFilter.filterDescription);
     }
 
     /**
@@ -111,7 +117,6 @@ type BuildHealthFilterOptionData = {
     id: string,
     text: string,
     query: string,
-    owner?: string,
     underylingObject: GetBuildHealthFilterResponse | undefined
 }
 
@@ -123,6 +128,32 @@ const ProjectFilterDropdownSingle: React.FC<{ handler: BuildHealthDataHandler, b
     let selectedFilterId = buildHealthOptions.getBuildHealthFilterId();
     let selectedFilter: GetBuildHealthFilterResponse | null = null;
     let selectedProject: string | null = buildHealthOptions.getCurrentProject();
+
+    let generateToolTip = (castedOption: IDropdownOption<BuildHealthFilterOptionData>): JSX.Element | undefined => {
+        if (!castedOption.data || castedOption.data.underylingObject === undefined) {
+            return handler.hasBuildHealthFilterWriteAccess ? (
+                <div>
+                    Select to make a new filter.
+                </div>
+            ) : undefined;
+        }
+
+        const owner = castedOption.data.underylingObject.owner?.name ?? "N/A";
+        const description = castedOption.data.underylingObject.filterDescription ?? "N/A";
+        const lastUpdated = getDateTimeString(new Date(castedOption.data.underylingObject.updateTimeUtc));
+
+        return (
+            <div>
+                <b>Owner</b>: {owner}
+                <br />
+                <b>Description</b>: {description}
+                <br />
+                <b>Last Updated</b>: {lastUpdated}
+                <br />
+            </div>
+        );
+    };
+
     if (selectedProject === null) {
         return (null);
     }
@@ -132,10 +163,12 @@ const ProjectFilterDropdownSingle: React.FC<{ handler: BuildHealthDataHandler, b
         return (null)
     }
 
+    buildHealthFilters.sort((a, b) => a.filterName.localeCompare(b.filterName));
+
     const doptions: IDropdownOption<BuildHealthFilterOptionData>[] = [];
     doptions.push({
         key: "empty",
-        text: "(New)",
+        text: handler.hasBuildHealthFilterWriteAccess ? "(New)" : "Select a Filter",
         data: undefined
     });
 
@@ -164,12 +197,12 @@ const ProjectFilterDropdownSingle: React.FC<{ handler: BuildHealthDataHandler, b
     return (
         <Stack>
             <TooltipHost
-                content="Saved filters for the selected project."
+                content="Saved filters for the selected project. Write access required to create new filters."
                 calloutProps={{ directionalHint: DirectionalHint.rightCenter }}>
                 <Label>Project Filters</Label>
             </TooltipHost>
             <Stack horizontal tokens={{ childrenGap: 8 }}>
-                <Dropdown
+                <SearchableDropdown
                     options={doptions}
                     styles={{ dropdown: { width: 300 } }}
                     selectedKey={computedSelectedKey}
@@ -188,7 +221,7 @@ const ProjectFilterDropdownSingle: React.FC<{ handler: BuildHealthDataHandler, b
                         }
 
                         return (
-                            <TooltipHost content={"Last Updated: " + (castedOption.data?.underylingObject !== undefined ? castedOption.data?.underylingObject?.updateTimeUtc : "N/A")}
+                            <TooltipHost content={generateToolTip(castedOption)}
                                 calloutProps={{ directionalHint: DirectionalHint.rightCenter }}>
                                 <div
                                     style={{
@@ -216,7 +249,7 @@ const ProjectFilterDropdownSingle: React.FC<{ handler: BuildHealthDataHandler, b
                         buildHealthOptions.setBuildHealthFilterId(option.data!.underylingObject!.id);
                         requestFilterApplication(query, buildHealthOptions, handler);
                     }}>
-                </Dropdown>
+                </SearchableDropdown>
                 {/* Only show the button if the user has write access for Build Health Filters. */}
                 {
                     handler.hasBuildHealthFilterWriteAccess
diff --git a/Engine/Source/Programs/Horde/HordeDashboard/src/components/buildhealth/BuildHealthFilterModal.tsx b/Engine/Source/Programs/Horde/HordeDashboard/src/components/buildhealth/BuildHealthFilterModal.tsx
--- a/Engine/Source/Programs/Horde/HordeDashboard/src/components/buildhealth/BuildHealthFilterModal.tsx
+++ b/Engine/Source/Programs/Horde/HordeDashboard/src/components/buildhealth/BuildHealthFilterModal.tsx
@@ -6,7 +6,6 @@ import { AddBuildHealthFiltersRequest, GetBuildHealthFilterResponse, UpdateBuild
 import { useState } from "react";
 import dashboard from "horde/backend/Dashboard";
 import { EditableBuildHealthFilter } from "./BuildHealthFilterComponent";
-import { getHordeStyling } from "horde/styles/Styles";
 import { getHordeTheme } from "horde/styles/theme";
 import { getDateTimeString } from "./stepoutcome/StepOutcomeUtilities";
 
@@ -58,7 +57,7 @@ export type BuildHealthFilterModalProps = {
     onRevert?: () => any;
 
     // The delegate to execute when the Build Health Filter has been deleted.
-    onDeleteSuccess?: (oldFilterProjectId? : string) => any;
+    onDeleteSuccess?: (oldFilterProjectId?: string) => any;
 };
 
 // #endregion -- Private Types --
@@ -77,22 +76,22 @@ export const BuildHealthFilterModal: React.FC<BuildHealthFilterModalProps> = ({
         return;
     }
 
-    const [state, setState] = useState<{ filterName: string }>(buildHealthFilter !== null && buildHealthFilter.filterName !== undefined ? { filterName: buildHealthFilter.filterName } : { filterName: "" });
+    const [state, setState] = useState<{ filterName: string, filterDescription?: string }>(buildHealthFilter !== null && (buildHealthFilter.filterName !== undefined || buildHealthFilter.filterDescription !== undefined) ? { filterName: buildHealthFilter.filterName, filterDescription: buildHealthFilter.filterDescription } : { filterName: "", filterDescription: undefined });
     const isNewFilter: boolean = buildHealthFilter.filterId === null;
-    const isEditorOwner: boolean = buildHealthFilter.filterOwner !== null && buildHealthFilter.filterOwner == dashboard.userId;
-    const isOldAndNewDiff: boolean = buildHealthFilter.oldHumanReadableQuery !== buildHealthFilter.newHumanReadableQuery || buildHealthFilter.filterName !== state.filterName;
+    const isEditorOwner: boolean = (buildHealthFilter.filterOwner !== null && buildHealthFilter.filterOwner.id == dashboard.userId) || dashboard.hordeAdmin;
+    const isOldAndNewDiff: boolean = buildHealthFilter.oldHumanReadableQuery !== buildHealthFilter.newHumanReadableQuery || buildHealthFilter.filterName !== state.filterName || buildHealthFilter.filterDescription !== state.filterDescription;
     const hasHumanReadanbleOldQuery: boolean = buildHealthFilter.oldHumanReadableQuery !== null;
 
     // #region -- Action Handlers --
 
     async function Save(): Promise<GetBuildHealthFilterResponse | null> {
         if (buildHealthFilter?.filterId === null) {
-            let request: AddBuildHealthFiltersRequest = { filterName: state.filterName, filterProject: buildHealthFilter.filterProjectId, filterQuery: buildHealthFilter.oldFilterQuery };
+            let request: AddBuildHealthFiltersRequest = { filterName: state.filterName, filterDescription: state.filterDescription, filterProject: buildHealthFilter.filterProjectId, filterQuery: buildHealthFilter.oldFilterQuery };
             let response: GetBuildHealthFilterResponse = await backend.addBuildHealthProjectFilter(request);
 
             return response;
         } else {
-            let request: UpdateBuildHealthFiltersRequest = { filterName: state.filterName, filterQuery: buildHealthFilter!.newFilterQuery };
+            let request: UpdateBuildHealthFiltersRequest = { filterName: state.filterName, filterDescription: state.filterDescription, filterQuery: buildHealthFilter!.newFilterQuery };
             let response: GetBuildHealthFilterResponse = await backend.updateBuildHealthProjectFilter(buildHealthFilter?.filterId!, request);
             return response;
         }
@@ -103,7 +102,7 @@ export const BuildHealthFilterModal: React.FC<BuildHealthFilterModalProps> = ({
             return null;
         }
 
-        let request: AddBuildHealthFiltersRequest = { filterName: state.filterName, filterProject: buildHealthFilter.filterProjectId, filterQuery: buildHealthFilter.newFilterQuery };
+        let request: AddBuildHealthFiltersRequest = { filterName: state.filterName, filterDescription: state.filterDescription, filterProject: buildHealthFilter.filterProjectId, filterQuery: buildHealthFilter.newFilterQuery };
         let response: GetBuildHealthFilterResponse = await backend.addBuildHealthProjectFilter(request);
 
         return response;
@@ -137,7 +136,7 @@ export const BuildHealthFilterModal: React.FC<BuildHealthFilterModalProps> = ({
         }
     };
 
-    const handleDeleteClick = async (oldFilterProjectId? : string) => {
+    const handleDeleteClick = async (oldFilterProjectId?: string) => {
         const savedFilter = await Delete();
         if (savedFilter) {
             onDeleteSuccess?.(oldFilterProjectId);
@@ -181,11 +180,27 @@ export const BuildHealthFilterModal: React.FC<BuildHealthFilterModalProps> = ({
             />
             <TextField
                 label="Filter Name"
+                maxLength={250}
                 value={state.filterName}
                 onChange={(_, newValue) => {
                     setState({ ...state, filterName: newValue ?? "" })
                 }}
             />
+            <TextField
+                label="Owner"
+                value={buildHealthFilter.filterOwner !== null ? buildHealthFilter.filterOwner.name : dashboard.username}
+                readOnly
+            />
+            <TextField
+                label="Description"
+                multiline
+                rows={5}
+                maxLength={500}
+                value={state.filterDescription ?? ""}
+                onChange={(_, newValue) => {
+                    setState({ ...state, filterDescription: newValue ?? undefined })
+                }}
+            />
             <Stack horizontal tokens={{ childrenGap: 16 }}>
                 {hasHumanReadanbleOldQuery && <Stack.Item grow>
                     <TextField
@@ -238,7 +253,7 @@ export const BuildHealthFilterModal: React.FC<BuildHealthFilterModalProps> = ({
             {(isNewFilter || isEditorOwner) && <PrimaryButton disabled={(!isNewFilter && !isOldAndNewDiff)} onClick={() => { handleSaveClick(); }} text="Save" />}
             {!isNewFilter && <PrimaryButton onClick={() => { handleSaveAsCopyClick(); }} text="Save As Copy" />}
             {!isNewFilter && isOldAndNewDiff && <PrimaryButton onClick={() => { handleRevertClick(); }} text="Revert" />}
-            {!isNewFilter && <PrimaryButton disabled={buildHealthFilter.filterOwner !== null && buildHealthFilter.filterOwner !== dashboard.userId} onClick={() => { handleDeleteClick(buildHealthFilter?.filterProjectId); }} text="Delete" />}
+            {!isNewFilter && <PrimaryButton disabled={buildHealthFilter.filterOwner !== null && buildHealthFilter.filterOwner.id !== dashboard.userId} onClick={() => { handleDeleteClick(buildHealthFilter?.filterProjectId); }} text="Delete" />}
         </Dialog >
     );
 };
diff --git a/Engine/Source/Programs/Horde/HordeDashboard/src/components/buildhealth/BuildhealthView.tsx b/Engine/Source/Programs/Horde/HordeDashboard/src/components/buildhealth/BuildhealthView.tsx
--- a/Engine/Source/Programs/Horde/HordeDashboard/src/components/buildhealth/BuildhealthView.tsx
+++ b/Engine/Source/Programs/Horde/HordeDashboard/src/components/buildhealth/BuildhealthView.tsx
@@ -1,6 +1,6 @@
 // Copyright Epic Games, Inc. All Rights Reserved.
 
-import { Checkbox, ComboBox, DefaultButton, DirectionalHint, Dropdown, DropdownMenuItemType, IconButton, IDropdownOption, IDropdownProps, ISelectableOption, Label, mergeStyleSets, PrimaryButton, SearchBox, SelectableOptionMenuItemType, Stack, TextField, TooltipHost } from "@fluentui/react";
+import { Checkbox, ComboBox, DefaultButton, DirectionalHint, Dropdown, DropdownMenuItemType, IconButton, IDropdownOption, ISelectableOption, Label, PrimaryButton, SelectableOptionMenuItemType, Stack, TooltipHost } from "@fluentui/react";
 import { GetBuildHealthFilterResponse, GetProjectResponse, GetStreamResponse } from "horde/backend/Api";
 import { observer } from "mobx-react-lite";
 import { useCallback, useEffect, useLayoutEffect, useMemo, useRef, useState } from "react";
@@ -11,7 +11,7 @@ import { Breadcrumbs } from "../Breadcrumbs";
 import { TopNav } from "../TopNav";
 import { StepOutcomeView } from "../buildhealth/stepoutcome/StepOutcome";
 import { BuildHealthDataHandler, DataHandlerRefreshRequest, EMPTY_REFRESH_REQUEST, expandRefreshRequest } from "./BuildHealthDataHandler";
-import { KEY_SEPARATOR, LabelRefData, PARAMETER_KEY_PREFIX, StepRefData, TemplateRefData } from "./BuildHealthDataTypes";
+import { dropdownOptions, KEY_SEPARATOR, LabelRefData, PARAMETER_KEY_PREFIX, StepRefData, TemplateRefData } from "./BuildHealthDataTypes";
 import { ALL_JOB_START_METHODS_LIST, ALL_STATES_TUPLE_LIST, BuildHealthOptionsState, BuildHealthOptionStateDiff, BuildHealthUIOptionsState, JobHistoryTimeSpans, loadBuildHealthOptionsFromParams, MANUAL_START_METHOD, parseBuildHealthQueryParams, SCHEDULED_START_METHOD, TimeSpan, toNavigationQuery, upgradeBuildHealthQueryParams } from "./BuildHealthOptions";
 import { decodeStepKey, decodeTemplateKey, encodeFullyQualifiedLabelName, encodeLabelKey, encodeStepKey, encodeTemplateKey } from "./BuildHealthUtilities";
 import { stepOutcomeTableClasses } from "./stepoutcome/StepOutcomeSharedUIComponents";
@@ -20,28 +20,10 @@ import { BuildHealthOptionsController, StreamDescriptor, TemplateDescriptor } fr
 import { StepOutcomeLegend } from "./stepoutcome/StepOutcomeLegend";
 import { StepOutcomeDataHandler, StepOutcomeFilters } from "./stepoutcome/StepOutcomeDataHandler";
 import { BuildHealthFilterComponent } from "./BuildHealthFilterComponent";
+import { SearchableDropdown } from "./SearchableDropdown";
 
 // #region -- Dropdown Components --
 
-// #region -- Styles --
-
-const dropdownNavigationButtons = mergeStyleSets({
-    base: {
-        marginLeft: 4,
-        padding: "2px 6px",
-        fontSize: 12,
-        cursor: "pointer",
-    }
-});
-
-const dropdownOptions = mergeStyleSets({
-    base: {
-        padding: "4px 8px",
-    }
-});
-
-// #endregion -- Styles --
-
 // #region -- Dropdown Data Types --
 
 type BuildHealthOptionData = {
@@ -62,177 +44,6 @@ const SELECT_CLEAR_ALL_SPECIAL_GROUP: string = "select_clear_all";
 
 // #region -- Base Dropdown Components --
 
-/**
- * Props for the SearchableDropdown component.
- */
-interface ISearchableDropdownProps extends IDropdownProps {
-
-    // Delegate that is run when search text field is changed.
-    onSearchValueChanged?: (searchValue: string) => void;
-    // Refs to all groups in the dropdown.
-    groupRefs?: React.RefObject<(HTMLDivElement | null)[]>;
-}
-
-/**
- * Searchable drop down component.
- * @param props The props that control the component.
- * @returns The React component.
- */
-const SearchableDropdown: React.FC<ISearchableDropdownProps> = (props) => {
-    const { modeColors } = getHordeStyling();
-    const { onSearchValueChanged, groupRefs, ...dropdownProps } = props;
-    const dropdownRef = useRef<HTMLDivElement | null>(null);
-    const currentIndex = useRef(0);
-    const PIXEL_AFFORDANCE: number = 5;
-
-    // #region -- Internal Dropdown Navigation Functions --
-
-    const recomputeCurrentIndex = (): { currentGroup: number, onHeader: boolean } => {
-        if (!dropdownRef.current || !groupRefs?.current) return { currentGroup: 0, onHeader: false };
-
-        const scrollTop = dropdownRef.current.scrollTop;
-        const keys = Object.keys(groupRefs.current || []);
-
-        let newIndex: number = 0;
-        let onHeader: boolean = false;
-
-        for (let i = 0; i < keys.length; i++) {
-            const el = groupRefs.current[keys[i]];
-            if (!el) continue;
-
-            const offset = el.offsetTop - 50;
-            if (scrollTop >= offset - PIXEL_AFFORDANCE) {
-                newIndex = i;
-
-                if (scrollTop >= offset - PIXEL_AFFORDANCE && scrollTop <= offset + PIXEL_AFFORDANCE) {
-                    onHeader = true;
-                }
-            }
-        }
-
-        currentIndex.current = newIndex;
-        return { currentGroup: newIndex, onHeader: onHeader };
-    };
-
-    const scrollToTop = () => {
-        currentIndex.current = 0;
-        goToTargetGroup(currentIndex.current);
-    };
-
-    const scrollToBottom = () => {
-        let groupRefKeys = groupRefs?.current ? Object.keys(groupRefs.current) : [];
-        currentIndex.current = groupRefKeys.length - 1;
-        goToTargetGroup(currentIndex.current);
-    };
-
-    const goToTargetGroup = (targetGroupIndex: number) => {
-        let groupRefKeys = groupRefs?.current ? Object.keys(groupRefs.current) : [];
-        if (!groupRefs?.current || !dropdownRef.current) {
-            return;
-        }
-        if (targetGroupIndex < 0 || targetGroupIndex >= groupRefKeys.length) {
-            return;
-        }
-
-        currentIndex.current = targetGroupIndex;
-        const el = groupRefs.current[groupRefKeys[targetGroupIndex]];
-        if (el) {
-            const offset = el.offsetTop - 50;
-            dropdownRef.current.scrollTop = offset >= 0 ? offset : 0;
-        }
-    }
-
-    const iterateToGroup = (index: number) => {
-        let currentLocation = recomputeCurrentIndex();
-        currentIndex.current = currentLocation.currentGroup;
-
-        // Determine how to move the current index.
-        // - If we are going "up" (-1 index), we need to see whether we are exactly on the header or not
-        // -- If we are exactly on the header, we intend to go to the previous group (so actually go to the previous group)
-        // -- If however we are within the group, we intend to snap to the current groups header
-        // - If we are going "down" (1 index), we simply just increment on top of the current group.
-        let incrementedIndex = index == -1 ? (currentLocation.onHeader ? currentIndex.current + index : currentIndex.current) : currentIndex.current + index;
-
-        goToTargetGroup(incrementedIndex);
-    };
-
-    // #endregion -- Internal Dropdown Navigation Functions --
-
-    return (
-        <Dropdown
-            ref={dropdownRef}
-            {...dropdownProps}
-            calloutProps={{
-                calloutWidth: 510,
-                directionalHint: DirectionalHint.rightTopEdge,
-                alignTargetEdge: true,
-                styles: {
-                    calloutMain: {
-                        width: 500,
-                        minWidth: 500,
-                        minHeight: 200,
-                        maxHeight: 800
-                    }
-                },
-            }}
-            onRenderList={(props, defaultRender) => {
-                if (!props) {
-                    return null;
-                }
-
-                return (
-                    <div ref={dropdownRef} style={{ maxHeight: 800, overflowY: 'auto', width: 500 }}>
-                        <Stack
-                            horizontal
-                            verticalAlign="center"
-                            styles={{ root: { position: "sticky", top: 0, zIndex: 1, padding: 8, background: modeColors.background } }}
-                        >
-                            <SearchBox
-                                placeholder="Filter options"
-                                autoFocus
-                                onChange={(_ev, newVal) => onSearchValueChanged?.(newVal !== undefined ? newVal : '')}
-                                onClear={() => onSearchValueChanged?.('')}
-                                onKeyDown={(ev) => {
-                                    if (ev.key === 'Enter' || ev.key === 'Escape') ev.stopPropagation();
-                                }}
-                                styles={{
-                                    root: {
-                                        flexGrow: 1,
-                                        marginRight: 4,
-                                    },
-                                    field: {
-                                    },
-                                }}
-                            />
-                            <IconButton
-                                iconProps={{ iconName: 'ChevronUp' }}
-                                title="Jump to previous group."
-                                className={dropdownNavigationButtons.base}
-                                onClick={() => { iterateToGroup(-1) }} />
-                            <IconButton
-                                iconProps={{ iconName: 'ChevronDown' }}
-                                title="Jump to next group."
-                                className={dropdownNavigationButtons.base}
-                                onClick={() => { iterateToGroup(1) }} />
-                            <IconButton
-                                iconProps={{ iconName: 'DoubleChevronUp' }}
-                                title="Jump to first group."
-                                className={dropdownNavigationButtons.base}
-                                onClick={scrollToTop} />
-                            <IconButton
-                                iconProps={{ iconName: 'DoubleChevronDown' }}
-                                title="Jump to last group."
-                                className={dropdownNavigationButtons.base}
-                                onClick={scrollToBottom} />
-                        </Stack>
-                        {defaultRender?.(props)}
-                    </div>
-                );
-            }}
-        />
-    );
-};
-
 /**
  * React Component that is a single select drop down.
  * @returns React Component.
diff --git a/Engine/Source/Programs/Horde/HordeDashboard/src/components/buildhealth/SearchableDropdownComponent.tsx b/Engine/Source/Programs/Horde/HordeDashboard/src/components/buildhealth/SearchableDropdownComponent.tsx
new file mode 100644
--- /dev/null
+++ b/Engine/Source/Programs/Horde/HordeDashboard/src/components/buildhealth/SearchableDropdownComponent.tsx
@@ -0,0 +1,175 @@
+import { DirectionalHint, Dropdown, IconButton, IDropdownProps, SearchBox, Stack } from "@fluentui/react";
+import { getHordeStyling } from "horde/styles/Styles";
+import { useRef } from "react";
+import { dropdownNavigationButtons } from "./BuildHealthDataTypes";
+
+/**
+ * Props for the SearchableDropdown component.
+ */
+export interface ISearchableDropdownProps extends IDropdownProps {
+
+    // Delegate that is run when search text field is changed.
+    onSearchValueChanged?: (searchValue: string) => void;
+    // Refs to all groups in the dropdown.
+    groupRefs?: React.RefObject<(HTMLDivElement | null)[]>;
+}
+
+/**
+ * Searchable drop down component.
+ * @param props The props that control the component.
+ * @returns The React component.
+ */
+export const SearchableDropdown: React.FC<ISearchableDropdownProps> = (props) => {
+    const { modeColors } = getHordeStyling();
+    const { onSearchValueChanged, groupRefs, ...dropdownProps } = props;
+    const dropdownRef = useRef<HTMLDivElement | null>(null);
+    const currentIndex = useRef(0);
+    const PIXEL_AFFORDANCE: number = 5;
+
+    // #region -- Internal Dropdown Navigation Functions --
+
+    const recomputeCurrentIndex = (): { currentGroup: number, onHeader: boolean } => {
+        if (!dropdownRef.current || !groupRefs?.current) return { currentGroup: 0, onHeader: false };
+
+        const scrollTop = dropdownRef.current.scrollTop;
+        const keys = Object.keys(groupRefs.current || []);
+
+        let newIndex: number = 0;
+        let onHeader: boolean = false;
+
+        for (let i = 0; i < keys.length; i++) {
+            const el = groupRefs.current[keys[i]];
+            if (!el) continue;
+
+            const offset = el.offsetTop - 50;
+            if (scrollTop >= offset - PIXEL_AFFORDANCE) {
+                newIndex = i;
+
+                if (scrollTop >= offset - PIXEL_AFFORDANCE && scrollTop <= offset + PIXEL_AFFORDANCE) {
+                    onHeader = true;
+                }
+            }
+        }
+
+        currentIndex.current = newIndex;
+        return { currentGroup: newIndex, onHeader: onHeader };
+    };
+
+    const scrollToTop = () => {
+        currentIndex.current = 0;
+        goToTargetGroup(currentIndex.current);
+    };
+
+    const scrollToBottom = () => {
+        let groupRefKeys = groupRefs?.current ? Object.keys(groupRefs.current) : [];
+        currentIndex.current = groupRefKeys.length - 1;
+        goToTargetGroup(currentIndex.current);
+    };
+
+    const goToTargetGroup = (targetGroupIndex: number) => {
+        let groupRefKeys = groupRefs?.current ? Object.keys(groupRefs.current) : [];
+        if (!groupRefs?.current || !dropdownRef.current) {
+            return;
+        }
+        if (targetGroupIndex < 0 || targetGroupIndex >= groupRefKeys.length) {
+            return;
+        }
+
+        currentIndex.current = targetGroupIndex;
+        const el = groupRefs.current[groupRefKeys[targetGroupIndex]];
+        if (el) {
+            const offset = el.offsetTop - 50;
+            dropdownRef.current.scrollTop = offset >= 0 ? offset : 0;
+        }
+    }
+
+    const iterateToGroup = (index: number) => {
+        let currentLocation = recomputeCurrentIndex();
+        currentIndex.current = currentLocation.currentGroup;
+
+        // Determine how to move the current index.
+        // - If we are going "up" (-1 index), we need to see whether we are exactly on the header or not
+        // -- If we are exactly on the header, we intend to go to the previous group (so actually go to the previous group)
+        // -- If however we are within the group, we intend to snap to the current groups header
+        // - If we are going "down" (1 index), we simply just increment on top of the current group.
+        let incrementedIndex = index == -1 ? (currentLocation.onHeader ? currentIndex.current + index : currentIndex.current) : currentIndex.current + index;
+
+        goToTargetGroup(incrementedIndex);
+    };
+
+    // #endregion -- Internal Dropdown Navigation Functions --
+
+    return (
+        <Dropdown
+            ref={dropdownRef}
+            {...dropdownProps}
+            calloutProps={{
+                calloutWidth: 510,
+                directionalHint: DirectionalHint.rightTopEdge,
+                alignTargetEdge: true,
+                styles: {
+                    calloutMain: {
+                        width: 500,
+                        minWidth: 500,
+                        minHeight: 200,
+                        maxHeight: 800
+                    }
+                },
+            }}
+            onRenderList={(props, defaultRender) => {
+                if (!props) {
+                    return null;
+                }
+
+                return (
+                    <div ref={dropdownRef} style={{ maxHeight: 800, overflowY: 'auto', width: 500 }}>
+                        <Stack
+                            horizontal
+                            verticalAlign="center"
+                            styles={{ root: { position: "sticky", top: 0, zIndex: 1, padding: 8, background: modeColors.background } }}
+                        >
+                            <SearchBox
+                                placeholder="Filter options"
+                                autoFocus
+                                onChange={(_ev, newVal) => onSearchValueChanged?.(newVal !== undefined ? newVal : '')}
+                                onClear={() => onSearchValueChanged?.('')}
+                                onKeyDown={(ev) => {
+                                    if (ev.key === 'Enter' || ev.key === 'Escape') ev.stopPropagation();
+                                }}
+                                styles={{
+                                    root: {
+                                        flexGrow: 1,
+                                        marginRight: 4,
+                                    },
+                                    field: {
+                                    },
+                                }}
+                            />
+                            <IconButton
+                                iconProps={{ iconName: 'ChevronUp' }}
+                                title="Jump to previous group."
+                                className={dropdownNavigationButtons.base}
+                                onClick={() => { iterateToGroup(-1) }} />
+                            <IconButton
+                                iconProps={{ iconName: 'ChevronDown' }}
+                                title="Jump to next group."
+                                className={dropdownNavigationButtons.base}
+                                onClick={() => { iterateToGroup(1) }} />
+                            <IconButton
+                                iconProps={{ iconName: 'DoubleChevronUp' }}
+                                title="Jump to first group."
+                                className={dropdownNavigationButtons.base}
+                                onClick={scrollToTop} />
+                            <IconButton
+                                iconProps={{ iconName: 'DoubleChevronDown' }}
+                                title="Jump to last group."
+                                className={dropdownNavigationButtons.base}
+                                onClick={scrollToBottom} />
+                        </Stack>
+                        {defaultRender?.(props)}
+                    </div>
+                );
+            }}
+        />
+    );
+};
\ No newline at end of file
diff --git a/Engine/Source/Programs/Horde/HordeDashboard/src/components/buildhealth/stepoutcome/StepOutcomeDataHandler.ts b/Engine/Source/Programs/Horde/HordeDashboard/src/components/buildhealth/stepoutcome/StepOutcomeDataHandler.ts
--- a/Engine/Source/Programs/Horde/HordeDashboard/src/components/buildhealth/stepoutcome/StepOutcomeDataHandler.ts
+++ b/Engine/Source/Programs/Horde/HordeDashboard/src/components/buildhealth/stepoutcome/StepOutcomeDataHandler.ts
@@ -1,10 +1,10 @@
 // Copyright Epic Games, Inc. All Rights Reserved.
 
-import { GetJobResponse, GetJobStepRefResponse, IssueQuery, JobState, JobStepOutcome, JobStepState, JobStreamQuery, LabelState } from "horde/backend/Api";
+import { GetJobResponse, GetJobStepRefResponse, IssueQuery, JobStepOutcome, JobStepState, JobStreamQuery, LabelState } from "horde/backend/Api";
 import { PollBase } from "horde/backend/PollBase";
 import { GetJobStepRefResponseWrapper, GetStepResponseWrapper, JobResponseData, LabelOutcomeTableEntry, StepOutcomeTable, StepOutcomeTableEntry, UpdatedChangeSummaryData } from "./StepOutcomeDataTypes";
 import backend from "horde/backend";
-import { decodeLabelKey, decodeStepKey, decodeTemplateKey, encodeLabelKeyFromStrings, encodeStepNameFromStrings, encodeStepNameFromStringsPreserveCase, encodeTemplateKeyFromStrings } from "../BuildHealthUtilities";
+import { decodeLabelKey, decodeStepKey, decodeTemplateKey, encodeLabelKeyFromStrings, encodeStepNameFromStringsPreserveCase, encodeTemplateKeyFromStrings } from "../BuildHealthUtilities";
 import { action, computed, makeObservable, observable, runInAction } from "mobx";
 import { summarizeJob, summarizeJobFromJobStepRef } from "./StepOutcomeUtilities";
 
@@ -242,7 +242,7 @@ export class StepOutcomeDataHandler extends PollBase {
                             if (templateStepSet &&
                                 templateStepSet.size >= 0 &&
                                 templateStepSet.has(fullyQualifiedStepName)) {
-                                this.stepOutcomeTableData.addEntry(new StepOutcomeTableEntry(new GetStepResponseWrapper(step), jobResponse.id, jobResponse.name, jobResponse.streamId, jobResponse.createTime.toString(), jobResponse.change, this.changeData.get(jobResponse.change!)?.dateUtc));
+                                this.stepOutcomeTableData.addEntry(new StepOutcomeTableEntry(new GetStepResponseWrapper(step), jobResponse.id, jobResponse.name, jobResponse.streamId, jobResponse.createTime.toString(), jobResponse.change, this.changeData.get(jobResponse.change!)?.dateUtc, batch));
                             }
                         }
                     }
diff --git a/Engine/Source/Programs/Horde/HordeDashboard/src/components/buildhealth/stepoutcome/StepOutcomeDataTypes.ts b/Engine/Source/Programs/Horde/HordeDashboard/src/components/buildhealth/stepoutcome/StepOutcomeDataTypes.ts
--- a/Engine/Source/Programs/Horde/HordeDashboard/src/components/buildhealth/stepoutcome/StepOutcomeDataTypes.ts
+++ b/Engine/Source/Programs/Horde/HordeDashboard/src/components/buildhealth/stepoutcome/StepOutcomeDataTypes.ts
@@ -1,6 +1,6 @@
 // Copyright Epic Games, Inc. All Rights Reserved.
 
-import { ChangeSummaryData, GetJobResponse, GetJobStepRefResponse, GetLabelStateResponse, GetStepResponse, IssueData, JobStepOutcome, JobStepState, LabelOutcome, LabelState } from "horde/backend/Api";
+import { ChangeSummaryData, GetBatchResponse, GetJobResponse, GetJobStepRefResponse, GetLabelStateResponse, GetStepResponse, IssueData, JobStepBatchError, JobStepOutcome, JobStepState, LabelOutcome, LabelState } from "horde/backend/Api";
 import { computeDuration, DURATION_NOT_SET, formatDurationFromMs, toDate } from "./StepOutcomeUtilities";
 import { StepOutcomeDataHandler, StepOutcomeViewOptions } from "./StepOutcomeDataHandler";
 import { action, computed, makeObservable, observable } from "mobx";
@@ -515,6 +515,7 @@ export class GetStepResponseWrapper implements IJobStepAdapter {
 export class StepOutcomeTableEntry extends OutcomeTableEntry {
     // #region -- Private Members --
 
+    private synthesizedStepOutcome: JobStepOutcome | undefined;
     private stepResponse: IJobStepAdapter;
     private stepDurationMs: number;
 
@@ -531,10 +532,11 @@ export class StepOutcomeTableEntry extends OutcomeTableEntry {
      * @param jobCreateTime The owning job's create time.
      * @param change The change associated with the table entry.
      * @param changeDate The change submit date associated with the table entry.
+     * @param batch The batch associated with the step.
      */
-    constructor(stepResponse: IJobStepAdapter, jobId: string, jobName: string, streamId: string, jobCreateTime: string, change?: number, changeDate?: Date | string) {
+    constructor(stepResponse: IJobStepAdapter, jobId: string, jobName: string, streamId: string, jobCreateTime: string, change?: number, changeDate?: Date | string, batch?: GetBatchResponse) {
         super(jobId, jobName, streamId, jobCreateTime, change, changeDate);
-
+        this.synthesizedStepOutcome = batch !== undefined && batch.error !== JobStepBatchError.None ? JobStepOutcome.Failure : undefined;
         this.stepResponse = stepResponse;
         this.stepDurationMs = computeDuration(stepResponse.getBase().startTime, stepResponse.getBase().finishTime)
     }
@@ -568,7 +570,7 @@ export class StepOutcomeTableEntry extends OutcomeTableEntry {
      * @inheritdoc
      */
     get outcome(): JobStepOutcome {
-        return this.stepResponse.getBase().outcome!;
+        return this.synthesizedStepOutcome !== undefined ? this.synthesizedStepOutcome : this.stepResponse.getBase().outcome!;
     }
 
     /**
@@ -950,7 +952,6 @@ export class StepOutcomeTable {
                     summaryTableEntry.stepsSkipped++;
                 }
 
-
                 if (cell.outcome == WarningStr) {
                     summaryTableEntry.stepsWarning++;
                 }
