[Horde][BuildHealth][Dashboard] Refreshed stale files from shelve for build break

#Synced-CL 49331470

diff --git a/Engine/Source/Programs/Horde/HordeDashboard/src/components/buildhealth/BuildHealthFilterComponent.tsx b/Engine/Source/Programs/Horde/HordeDashboard/src/components/buildhealth/BuildHealthFilterComponent.tsx
--- a/Engine/Source/Programs/Horde/HordeDashboard/src/components/buildhealth/BuildHealthFilterComponent.tsx
+++ b/Engine/Source/Programs/Horde/HordeDashboard/src/components/buildhealth/BuildHealthFilterComponent.tsx
@@ -7,8 +7,8 @@ import { observer } from "mobx-react-lite";
 import { BuildHealthFilterModal, BuildHealthFilterModalProps } from "./BuildHealthFilterModal";
 import { parseBuildHealthQueryParams, upgradeBuildHealthQueryParams, loadBuildHealthOptionsFromParams, toNavigationQuery } from "./BuildHealthOptions";
 import { getHordeStyling } from "horde/styles/Styles";
-import { SearchableDropdown } from "./SearchableDropdown";
 import { getDateTimeString } from "./stepoutcome/StepOutcomeUtilities";
+import { SearchableDropdown } from "./SearchableDropdownComponent";
 
 // #region -- Helper Functions --
 
diff --git a/Engine/Source/Programs/Horde/HordeDashboard/src/components/buildhealth/BuildhealthView.tsx b/Engine/Source/Programs/Horde/HordeDashboard/src/components/buildhealth/BuildhealthView.tsx
--- a/Engine/Source/Programs/Horde/HordeDashboard/src/components/buildhealth/BuildhealthView.tsx
+++ b/Engine/Source/Programs/Horde/HordeDashboard/src/components/buildhealth/BuildhealthView.tsx
@@ -20,7 +20,7 @@ import { BuildHealthOptionsController, StreamDescriptor, TemplateDescriptor } fr
 import { StepOutcomeLegend } from "./stepoutcome/StepOutcomeLegend";
 import { StepOutcomeDataHandler, StepOutcomeFilters } from "./stepoutcome/StepOutcomeDataHandler";
 import { BuildHealthFilterComponent } from "./BuildHealthFilterComponent";
-import { SearchableDropdown } from "./SearchableDropdown";
+import { SearchableDropdown } from "./SearchableDropdownComponent";
 
 // #region -- Dropdown Components --
 
