Horde: Instrument gate handling for schedules

diff --git a/Engine/Source/Programs/Horde/Plugins/Build/HordeServer.Build/Jobs/Schedules/ScheduleService.cs b/Engine/Source/Programs/Horde/Plugins/Build/HordeServer.Build/Jobs/Schedules/ScheduleService.cs
--- a/Engine/Source/Programs/Horde/Plugins/Build/HordeServer.Build/Jobs/Schedules/ScheduleService.cs
+++ b/Engine/Source/Programs/Horde/Plugins/Build/HordeServer.Build/Jobs/Schedules/ScheduleService.cs
@@ -14,7 +14,6 @@ using EpicGames.Redis.Utility;
 using EpicGames.Serialization;
 using HordeServer.Acls;
 using HordeServer.Auditing;
-using HordeServer.Commits;
 using HordeServer.Jobs.Graphs;
 using HordeServer.Jobs.Templates;
 using HordeServer.Server;
@@ -68,7 +67,6 @@ namespace HordeServer.Jobs.Schedules
 
 		readonly IAuditLog<ScheduleId> _auditLog;
 		readonly IGraphCollection _graphs;
-		readonly ICommitService _commitService;
 		readonly IJobCollection _jobCollection;
 		readonly IDowntimeService _downtimeService;
 		readonly JobService _jobService;
@@ -87,11 +85,10 @@ namespace HordeServer.Jobs.Schedules
 		/// <summary>
 		/// Constructor
 		/// </summary>
-		public ScheduleService(IAuditLog<ScheduleId> auditLog, IRedisService redis, IGraphCollection graphs, ICommitService commitService, IJobCollection jobCollection, JobService jobService, IDowntimeService downtimeService, IStreamCollection streamCollection, ITemplateCollectionInternal templateCollection, IMongoService mongoService, IClock clock, IOptionsMonitor<BuildConfig> buildConfig, Tracer tracer, ILogger<ScheduleService> logger)
+		public ScheduleService(IAuditLog<ScheduleId> auditLog, IRedisService redis, IGraphCollection graphs, IJobCollection jobCollection, JobService jobService, IDowntimeService downtimeService, IStreamCollection streamCollection, ITemplateCollectionInternal templateCollection, IMongoService mongoService, IClock clock, IOptionsMonitor<BuildConfig> buildConfig, Tracer tracer, ILogger<ScheduleService> logger)
 		{
 			_auditLog = auditLog;
 			_graphs = graphs;
-			_commitService = commitService;
 			_jobCollection = jobCollection;
 			_jobService = jobService;
 			_downtimeService = downtimeService;
@@ -339,7 +336,7 @@ namespace HordeServer.Jobs.Schedules
 			// Trigger this schedule
 			try
 			{
-				await TriggerAsync(stream, templateId, templateRef, schedule, schedule.ActiveJobs.Count - removeJobIds.Count, utcNow, scheduleLogger, cancellationToken);
+				await TriggerAsync(stream, templateId, templateRef, schedule, schedule.ActiveJobs.Count - removeJobIds.Count, utcNow, scheduleLogger, span, cancellationToken);
 			}
 			catch (Exception ex)
 			{
@@ -362,11 +359,13 @@ namespace HordeServer.Jobs.Schedules
 		/// <param name="numActiveJobs"></param>
 		/// <param name="utcNow">The current time</param>
 		/// <param name="scheduleLogger"></param>
+		/// <param name="linkedSpan">Span to link with</param>
 		/// <param name="cancellationToken"></param>
 		/// <returns>Async task</returns>
-		private async Task TriggerAsync(IStream stream, TemplateId templateId, ITemplateRef templateRef, ISchedule schedule, int numActiveJobs, DateTime utcNow, ILogger scheduleLogger, CancellationToken cancellationToken)
+		private async Task TriggerAsync(IStream stream, TemplateId templateId, ITemplateRef templateRef, ISchedule schedule, int numActiveJobs, DateTime utcNow, ILogger scheduleLogger, TelemetrySpan linkedSpan, CancellationToken cancellationToken)
 		{
-			using TelemetrySpan span = _tracer.StartActiveSpan($"{nameof(ScheduleService)}.{nameof(TriggerAsync)}.Full");
+			// Create a new root span as it may contain lots of child spans, this to prevent a huge single trace/span
+			using TelemetrySpan span = _tracer.StartRootSpan($"{nameof(ScheduleService)}.{nameof(TriggerAsync)}.Full", links: [new Link(linkedSpan.Context)]);
 			span.SetAttribute("horde.stream_id", stream.Id);
 			span.SetAttribute("horde.template_id", templateId);
 			span.SetAttribute("horde.active_jobs", numActiveJobs);
@@ -582,26 +581,35 @@ namespace HordeServer.Jobs.Schedules
 		/// </summary>
 		private async Task<CommitIdWithOrder?> GetNextChangeForGateAsync(StreamId streamId, TemplateId templateRefId, IScheduleGate gate, CommitIdWithOrder? minCommitId, CommitIdWithOrder? maxCommitId, ILogger scheduleLogger, CancellationToken cancellationToken)
 		{
-			for (; ; )
+			using TelemetrySpan span = _tracer.StartRootSpan($"{nameof(ScheduleService)}.{nameof(GetNextChangeForGateAsync)}");
+			span.SetAttribute("horde.gate.template_id", gate.TemplateId);
+			span.SetAttribute("horde.gate.target", gate.Target);
+			span.SetAttribute("horde.min_commit_id", minCommitId?.ToString());
+			span.SetAttribute("horde.max_commit_id", maxCommitId?.ToString());
+			
+			for (int c = 0; ; c++)
 			{
 				cancellationToken.ThrowIfCancellationRequested();
+				span.SetAttribute("horde.iteration", c);
+				span.AddEvent("Fetching job", new SpanAttributes(new KeyValuePair<string, object?>[] { new("horde.max_commit_id", maxCommitId) }));
 
-				IReadOnlyList<IJob> jobs = await _jobCollection.FindAsync(new FindJobOptions(StreamId: streamId, Templates: new[] { gate.TemplateId }, MinCommitId: minCommitId, MaxCommitId: maxCommitId), count: 2, cancellationToken: cancellationToken);
+				FindJobOptions options = new (StreamId: streamId, Templates: [gate.TemplateId], MinCommitId: minCommitId, MaxCommitId: maxCommitId);
+				IReadOnlyList<IJob> jobs = await _jobCollection.FindAsync(options, count: 2, cancellationToken: cancellationToken);
 
 				IJob? job = jobs.FirstOrDefault(x => maxCommitId == null || x.CommitId < maxCommitId);
 				if (job == null)
 				{
+					span.AddEvent("Job not found");
 					return null;
 				}
 
 				IGraph? graph = await _graphs.GetAsync(job.GraphHash, cancellationToken);
 				if (graph != null)
 				{
-					(JobStepState, JobStepOutcome)? state = job.GetTargetState(graph, gate.Target);
-					if (state != null && state.Value.Item1 == JobStepState.Completed)
+					(JobStepState stepState, JobStepOutcome stepOutcome)? state = job.GetTargetState(graph, gate.Target);
+					if (state is { stepState: JobStepState.Completed })
 					{
-						JobStepOutcome outcome = state.Value.Item2;
-						if (outcome == JobStepOutcome.Success || outcome == JobStepOutcome.Warnings)
+						if (state is { stepOutcome: JobStepOutcome.Success or JobStepOutcome.Warnings})
 						{
 							return job.CommitId;
 						}
