Horde: Updating schema documentation

#rnx

diff --git a/Engine/Source/Programs/Horde/Docs/Config/Schema/AclActions.md b/Engine/Source/Programs/Horde/Docs/Config/Schema/AclActions.md
--- a/Engine/Source/Programs/Horde/Docs/Config/Schema/AclActions.md
+++ b/Engine/Source/Programs/Horde/Docs/Config/Schema/AclActions.md
@@ -37,48 +37,6 @@
 | `DeleteAccount` |  |
 | `ViewAccount` |  |
 
-## Artifacts
-
-| Name | Description |
-| ---- | ----------- |
-| `ReadArtifact` | Permission to read from an artifact |
-| `WriteArtifact` | Permission to write to an artifact |
-| `DeleteArtifact` | Permission to delete to an artifact |
-| `UploadArtifact` | Ability to create an artifact. Typically just for debugging; agents have this access for a particular session. |
-| `DownloadArtifact` | Ability to download an artifact |
-
-## Bisect
-
-| Name | Description |
-| ---- | ----------- |
-| `CreateBisectTask` | Ability to start new bisect tasks |
-| `UpdateBisectTask` | Ability to update a bisect task |
-| `ViewBisectTask` | Ability to view a bisect task |
-
-## BuildHealth
-
-| Name | Description |
-| ---- | ----------- |
-| `AddUpdateDeleteBuildHealthFilter` | Ability to add, update, and delete build health filters. |
-
-## Devices
-
-| Name | Description |
-| ---- | ----------- |
-| `DeviceRead` | Ability to read devices |
-| `DeviceWrite` | Ability to write devices |
-
-## Jobs
-
-| Name | Description |
-| ---- | ----------- |
-| `CreateJob` | Ability to start new jobs |
-| `UpdateJob` | Rename a job, modify its priority, etc... |
-| `DeleteJob` | Delete a job properties |
-| `ExecuteJob` | Allows updating a job metadata (name, changelist number, step properties, new groups, job states, etc...). Typically granted to agents. Not user facing. |
-| `RetryJobStep` | Ability to retry a failed job step |
-| `ViewJob` | Ability to view a job |
-
 ## Notices
 
 | Name | Description |
@@ -87,28 +45,6 @@
 | `UpdateNotice` |  |
 | `DeleteNotice` |  |
 
-## Notifications
-
-| Name | Description |
-| ---- | ----------- |
-| `CreateSubscription` | Ability to subscribe to notifications |
-
-## Projects
-
-| Name | Description |
-| ---- | ----------- |
-| `CreateProject` | Allows the creation of new projects |
-| `DeleteProject` | Allows deletion of projects. |
-| `UpdateProject` | Modify attributes of a project (name, categories, etc...) |
-| `ViewProject` | View information about a project |
-
-## Replicators
-
-| Name | Description |
-| ---- | ----------- |
-| `UpdateReplicator` | Allows deletion of projects. |
-| `ViewReplicator` | Allows the creation of new projects |
-
 ## ServiceAccounts
 
 | Name | Description |
@@ -118,24 +54,6 @@
 | `DeleteAccount` |  |
 | `ViewAccount` |  |
 
-## Streams
-
-| Name | Description |
-| ---- | ----------- |
-| `CreateStream` | Allows the creation of new streams within a project |
-| `UpdateStream` | Allows updating a stream (agent types, templates, schedules) |
-| `DeleteStream` | Allows deleting a stream |
-| `ViewStream` | Ability to view a stream |
-| `ViewChanges` | View changes submitted to a stream. NOTE: this returns responses from the server's Perforce account, which may be a priviledged user. |
-| `ViewTemplate` | View template associated with a stream |
-
-## TestData
-
-| Name | Description |
-| ---- | ----------- |
-| `TestHealthRead` | Ability to read test health |
-| `TestHealthWrite` | Ability to write test health |
-
 ## Accounts
 
 | Name | Description |
@@ -145,47 +63,43 @@
 | `DeleteAccount` |  |
 | `ViewAccount` |  |
 
-## Artifacts
+## Agents
 
 | Name | Description |
 | ---- | ----------- |
-| `ReadArtifact` |  |
-| `WriteArtifact` |  |
-| `DeleteArtifact` |  |
-| `UploadArtifact` |  |
-| `DownloadArtifact` |  |
+| `CreateAgent` | Ability to create an agent. This may be done explicitly, or granted to agents to allow them to self-register. |
+| `CreateWorkstationAgent` | Ability to create a workstation agent. Granted to users with hardware to offer. |
+| `UpdateAgent` | Update an agent's name, pools, etc... |
+| `DeleteAgent` | Soft-delete an agent |
+| `ViewAgent` | View an agent |
+| `ListAgents` | List the available agents |
 
-## Bisect
+## Compute
 
 | Name | Description |
 | ---- | ----------- |
-| `CreateBisectTask` |  |
-| `UpdateBisectTask` |  |
-| `ViewBisectTask` |  |
+| `AddComputeTasks` | User can add tasks to the compute cluster |
+| `GetComputeTasks` | User can get and list tasks from the compute cluster |
+| `UbaCacheRead` | User can read from Unreal Build Accelerator cache |
+| `UbaCacheWrite` | User can write to Unreal Build Accelerator cache |
 
-## BuildHealth
+## Leases
 
 | Name | Description |
 | ---- | ----------- |
-| `AddUpdateDeleteBuildHealthFilter` |  |
+| `ViewLeases` | View all the leases that an agent has worked on |
+| `ViewLeaseTasks` | View the task data for a lease |
 
-## Devices
+## Logs
 
 | Name | Description |
 | ---- | ----------- |
-| `DeviceRead` |  |
-| `DeviceWrite` |  |
+| `CreateLog` | Ability to create a log. Implicitly granted to agents. |
+| `UpdateLog` | Ability to update log metadata |
+| `ViewLog` | Ability to view a log contents |
+| `WriteLogData` | Ability to write log data |
+| `CreateEvent` | Ability to create events |
+| `ViewEvent` | Ability to view events |
 
-## Jobs
-
-| Name | Description |
-| ---- | ----------- |
-| `CreateJob` |  |
-| `UpdateJob` |  |
-| `DeleteJob` |  |
-| `ExecuteJob` |  |
-| `RetryJobStep` |  |
-| `ViewJob` |  |
-
 ## Notices
 
@@ -195,27 +109,15 @@
 | `UpdateNotice` |  |
 | `DeleteNotice` |  |
 
-## Notifications
+## Pools
 
 | Name | Description |
 | ---- | ----------- |
-| `CreateSubscription` |  |
+| `CreatePool` | Create a global pool of agents |
+| `UpdatePool` | Modify an agent pool |
+| `DeletePool` | Delete an agent pool |
+| `ViewPool` | Ability to view a pool |
+| `ListPools` | View all the available agent pools |
 
-## Projects
-
-| Name | Description |
-| ---- | ----------- |
-| `CreateProject` |  |
-| `DeleteProject` |  |
-| `UpdateProject` |  |
-| `ViewProject` |  |
-
-## Replicators
-
-| Name | Description |
-| ---- | ----------- |
-| `UpdateReplicator` |  |
-| `ViewReplicator` |  |
-
 ## ServiceAccounts
 
@@ -226,23 +128,20 @@
 | `DeleteAccount` |  |
 | `ViewAccount` |  |
 
-## Streams
+## Sessions
 
 | Name | Description |
 | ---- | ----------- |
-| `CreateStream` |  |
-| `UpdateStream` |  |
-| `DeleteStream` |  |
-| `ViewStream` |  |
-| `ViewChanges` |  |
-| `ViewTemplate` |  |
+| `CreateSession` | Granted to agents to call CreateSession, which returns a bearer token identifying themselves valid to call UpdateSesssion via gRPC. |
+| `ViewSession` | Allows viewing information about an agent session |
 
-## TestData
+## Software
 
 | Name | Description |
 | ---- | ----------- |
-| `TestHealthRead` |  |
-| `TestHealthWrite` |  |
+| `UploadSoftware` | Ability to upload new versions of the agent software |
+| `DownloadSoftware` | Ability to download the agent software |
+| `DeleteSoftware` | Ability to delete agent software |
 
 ## Accounts
 
@@ -257,81 +156,81 @@
 
 | Name | Description |
 | ---- | ----------- |
-| `CreateAgent` | Ability to create an agent. This may be done explicitly, or granted to agents to allow them to self-register. |
-| `CreateWorkstationAgent` | Ability to create a workstation agent. Granted to users with hardware to offer. |
-| `UpdateAgent` | Update an agent's name, pools, etc... |
-| `DeleteAgent` | Soft-delete an agent |
-| `ViewAgent` | View an agent |
-| `ListAgents` | List the available agents |
+| `CreateAgent` |  |
+| `CreateWorkstationAgent` |  |
+| `UpdateAgent` |  |
+| `DeleteAgent` |  |
+| `ViewAgent` |  |
+| `ListAgents` |  |
 
 ## Artifacts
 
 | Name | Description |
 | ---- | ----------- |
-| `ReadArtifact` |  |
-| `WriteArtifact` |  |
-| `DeleteArtifact` |  |
-| `UploadArtifact` |  |
-| `DownloadArtifact` |  |
+| `ReadArtifact` | Permission to read from an artifact |
+| `WriteArtifact` | Permission to write to an artifact |
+| `DeleteArtifact` | Permission to delete to an artifact |
+| `UploadArtifact` | Ability to create an artifact. Typically just for debugging; agents have this access for a particular session. |
+| `DownloadArtifact` | Ability to download an artifact |
 
 ## Bisect
 
 | Name | Description |
 | ---- | ----------- |
-| `CreateBisectTask` |  |
-| `UpdateBisectTask` |  |
-| `ViewBisectTask` |  |
+| `CreateBisectTask` | Ability to start new bisect tasks |
+| `UpdateBisectTask` | Ability to update a bisect task |
+| `ViewBisectTask` | Ability to view a bisect task |
 
 ## BuildHealth
 
 | Name | Description |
 | ---- | ----------- |
-| `AddUpdateDeleteBuildHealthFilter` |  |
+| `AddUpdateDeleteBuildHealthFilter` | Ability to add, update, and delete build health filters. |
 
 ## Compute
 
 | Name | Description |
 | ---- | ----------- |
-| `AddComputeTasks` | User can add tasks to the compute cluster |
-| `GetComputeTasks` | User can get and list tasks from the compute cluster |
-| `UbaCacheRead` | User can read from Unreal Build Accelerator cache |
-| `UbaCacheWrite` | User can write to Unreal Build Accelerator cache |
+| `AddComputeTasks` |  |
+| `GetComputeTasks` |  |
+| `UbaCacheRead` |  |
+| `UbaCacheWrite` |  |
 
 ## Devices
 
 | Name | Description |
 | ---- | ----------- |
-| `DeviceRead` |  |
-| `DeviceWrite` |  |
+| `DeviceRead` | Ability to read devices |
+| `DeviceWrite` | Ability to write devices |
 
 ## Jobs
 
 | Name | Description |
 | ---- | ----------- |
-| `CreateJob` |  |
-| `UpdateJob` |  |
-| `DeleteJob` |  |
-| `ExecuteJob` |  |
-| `RetryJobStep` |  |
-| `ViewJob` |  |
+| `CreateJob` | Ability to start new jobs |
+| `UpdateJob` | Rename a job, modify its priority, etc... |
+| `DeleteJob` | Delete a job properties |
+| `ExecuteJob` | Allows updating a job metadata (name, changelist number, step properties, new groups, job states, etc...). Typically granted to agents. Not user facing. |
+| `RetryJobStep` | Ability to retry a failed job step |
+| `ViewJob` | Ability to view a job |
 
 ## Leases
 
 | Name | Description |
 | ---- | ----------- |
-| `ViewLeases` | View all the leases that an agent has worked on |
-| `ViewLeaseTasks` | View the task data for a lease |
+| `ViewLeases` |  |
+| `ViewLeaseTasks` |  |
 
 ## Logs
 
 | Name | Description |
 | ---- | ----------- |
-| `CreateLog` | Ability to create a log. Implicitly granted to agents. |
-| `UpdateLog` | Ability to update log metadata |
-| `ViewLog` | Ability to view a log contents |
-| `WriteLogData` | Ability to write log data |
-| `CreateEvent` | Ability to create events |
-| `ViewEvent` | Ability to view events |
+| `CreateLog` |  |
+| `UpdateLog` |  |
+| `ViewLog` |  |
+| `WriteLogData` |  |
+| `CreateEvent` |  |
+| `ViewEvent` |  |
 
 ## Notices
 
@@ -345,33 +244,33 @@
 
 | Name | Description |
 | ---- | ----------- |
-| `CreateSubscription` |  |
+| `CreateSubscription` | Ability to subscribe to notifications |
 
 ## Pools
 
 | Name | Description |
 | ---- | ----------- |
-| `CreatePool` | Create a global pool of agents |
-| `UpdatePool` | Modify an agent pool |
-| `DeletePool` | Delete an agent pool |
-| `ViewPool` | Ability to view a pool |
-| `ListPools` | View all the available agent pools |
+| `CreatePool` |  |
+| `UpdatePool` |  |
+| `DeletePool` |  |
+| `ViewPool` |  |
+| `ListPools` |  |
 
 ## Projects
 
 | Name | Description |
 | ---- | ----------- |
-| `CreateProject` |  |
-| `DeleteProject` |  |
-| `UpdateProject` |  |
-| `ViewProject` |  |
+| `CreateProject` | Allows the creation of new projects |
+| `DeleteProject` | Allows deletion of projects. |
+| `UpdateProject` | Modify attributes of a project (name, categories, etc...) |
+| `ViewProject` | View information about a project |
 
 ## Replicators
 
 | Name | Description |
 | ---- | ----------- |
-| `UpdateReplicator` |  |
-| `ViewReplicator` |  |
+| `UpdateReplicator` | Allows deletion of projects. |
+| `ViewReplicator` | Allows the creation of new projects |
 
 ## ServiceAccounts
 
@@ -386,34 +285,34 @@
 
 | Name | Description |
 | ---- | ----------- |
-| `CreateSession` | Granted to agents to call CreateSession, which returns a bearer token identifying themselves valid to call UpdateSesssion via gRPC. |
-| `ViewSession` | Allows viewing information about an agent session |
+| `CreateSession` |  |
+| `ViewSession` |  |
 
 ## Software
 
 | Name | Description |
 | ---- | ----------- |
-| `UploadSoftware` | Ability to upload new versions of the agent software |
-| `DownloadSoftware` | Ability to download the agent software |
-| `DeleteSoftware` | Ability to delete agent software |
+| `UploadSoftware` |  |
+| `DownloadSoftware` |  |
+| `DeleteSoftware` |  |
 
 ## Streams
 
 | Name | Description |
 | ---- | ----------- |
-| `CreateStream` |  |
-| `UpdateStream` |  |
-| `DeleteStream` |  |
-| `ViewStream` |  |
-| `ViewChanges` |  |
-| `ViewTemplate` |  |
+| `CreateStream` | Allows the creation of new streams within a project |
+| `UpdateStream` | Allows updating a stream (agent types, templates, schedules) |
+| `DeleteStream` | Allows deleting a stream |
+| `ViewStream` | Ability to view a stream |
+| `ViewChanges` | View changes submitted to a stream. NOTE: this returns responses from the server's Perforce account, which may be a priviledged user. |
+| `ViewTemplate` | View template associated with a stream |
 
 ## TestData
 
 | Name | Description |
 | ---- | ----------- |
-| `TestHealthRead` |  |
-| `TestHealthWrite` |  |
+| `TestHealthRead` | Ability to read test health |
+| `TestHealthWrite` | Ability to write test health |
 
 ## Accounts
 
@@ -886,6 +785,12 @@
 | `UpdateReplicator` |  |
 | `ViewReplicator` |  |
 
+## Secrets
+
+| Name | Description |
+| ---- | ----------- |
+| `ViewSecret` | View a credential |
+
 ## ServiceAccounts
 
 | Name | Description |
@@ -910,19 +815,6 @@
 | `DownloadSoftware` |  |
 | `DeleteSoftware` |  |
 
-## Storage
-
-| Name | Description |
-| ---- | ----------- |
-| `ReadBlobs` | Ability to read blobs from the storage service |
-| `WriteBlobs` | Ability to write blobs to the storage service |
-| `ReadAliases` | Ability to read aliases from the storage service |
-| `WriteAliases` | Ability to write aliases to the storage service |
-| `DeleteAliases` | Ability to write aliases to the storage service |
-| `ReadRefs` | Ability to read refs from the storage service |
-| `WriteRefs` | Ability to write refs to the storage service |
-| `DeleteRefs` | Ability to delete refs |
-
 ## Streams
 
 | Name | Description |
@@ -1074,7 +966,7 @@
 
 | Name | Description |
 | ---- | ----------- |
-| `ViewSecret` | View a credential |
+| `ViewSecret` |  |
 
 ## ServiceAccounts
 
@@ -1104,14 +996,14 @@
 
 | Name | Description |
 | ---- | ----------- |
-| `ReadBlobs` |  |
-| `WriteBlobs` |  |
-| `ReadAliases` |  |
-| `WriteAliases` |  |
-| `DeleteAliases` |  |
-| `ReadRefs` |  |
-| `WriteRefs` |  |
-| `DeleteRefs` |  |
+| `ReadBlobs` | Ability to read blobs from the storage service |
+| `WriteBlobs` | Ability to write blobs to the storage service |
+| `ReadAliases` | Ability to read aliases from the storage service |
+| `WriteAliases` | Ability to write aliases to the storage service |
+| `DeleteAliases` | Ability to write aliases to the storage service |
+| `ReadRefs` | Ability to read refs from the storage service |
+| `WriteRefs` | Ability to write refs to the storage service |
+| `DeleteRefs` | Ability to delete refs |
 
 ## Streams
 
@@ -1314,12 +1206,6 @@
 | `ViewChanges` |  |
 | `ViewTemplate` |  |
 
-## Symbols
-
-| Name | Description |
-| ---- | ----------- |
-| `ReadSymbols` | Ability to download symbols |
-
 ## TestData
 
 | Name | Description |
@@ -1327,5 +1213,12 @@
 | `TestHealthRead` |  |
 | `TestHealthWrite` |  |
 
+## Tools
+
+| Name | Description |
+| ---- | ----------- |
+| `DownloadTool` | Ability to download a tool |
+| `UploadTool` | Ability to upload new tool versions |
+
 ## Accounts
 
@@ -1514,7 +1407,7 @@
 
 | Name | Description |
 | ---- | ----------- |
-| `ReadSymbols` |  |
+| `ReadSymbols` | Ability to download symbols |
 
 ## TestData
 
@@ -1527,8 +1420,8 @@
 
 | Name | Description |
 | ---- | ----------- |
-| `DownloadTool` | Ability to download a tool |
-| `UploadTool` | Ability to upload new tool versions |
+| `DownloadTool` |  |
+| `UploadTool` |  |
 
 ## Accounts
 
diff --git a/Engine/Source/Programs/Horde/Docs/Config/Schema/Globals.md b/Engine/Source/Programs/Horde/Docs/Config/Schema/Globals.md
--- a/Engine/Source/Programs/Horde/Docs/Config/Schema/Globals.md
+++ b/Engine/Source/Programs/Horde/Docs/Config/Schema/Globals.md
@@ -59,137 +59,37 @@
 
 Name | Description
 ---- | -----------
-`build` | [BuildConfig](#buildconfig)<br>Configuration for the build plugin
 `analytics` | [AnalyticsConfig](#analyticsconfig)<br>Configuration for the analytics plugin
 `compute` | [ComputeConfig](#computeconfig)<br>Configuration for the compute plugin
+`build` | [BuildConfig](#buildconfig)<br>Configuration for the build plugin
 `exampleplugin` | [ExamplePluginConfig](#examplepluginconfig)<br>Configuration for the exampleplugin plugin
 `experimental` | [ExperimentalConfig](#experimentalconfig)<br>Configuration for the experimental plugin
-`storage` | [StorageConfig](#storageconfig)<br>Configuration for the storage plugin
 `secrets` | [SecretsConfig](#secretsconfig)<br>Configuration for the secrets plugin
-`symbols` | [SymbolsConfig](#symbolsconfig)<br>Configuration for the symbols plugin
+`storage` | [StorageConfig](#storageconfig)<br>Configuration for the storage plugin
 `tools` | [ToolsConfig](#toolsconfig)<br>Configuration for the tools plugin
+`symbols` | [SymbolsConfig](#symbolsconfig)<br>Configuration for the symbols plugin
 
-## BuildConfig
+## AnalyticsConfig
 
-Configuration for the build plugin
+Config settings for analytics
 
 Name | Description
 ---- | -----------
-`perforceClusters` | [PerforceCluster](#perforcecluster)`[]`<br>List of Perforce clusters
-`devices` | [DeviceConfig](#deviceconfig)<br>Device configuration
-`maxConformCount` | `integer`<br>Maximum number of conforms to run at once
-`agentShutdownIfDisabledGracePeriod` | `string`<br>Time to wait before shutting down an agent that has been disabled Used if no value is set on the actual pool.
-`artifactTypes` | [ArtifactTypeConfig](#artifacttypeconfig)`[]`<br>Configuration for different artifact types
-`projects` | [ProjectConfig](Projects.md)`[]`<br>List of projects
-`enableConformTasks` | `boolean`<br>Whether to allow conform tasks to run
-`enableJobTaskConformCheck` | `boolean`<br>Whether to check if a conform is required on an agent before assigning a lease for a job task source.
-`issueFixedTag` | `string`<br>Commit tag to use for marking issues as fixed
-`issues` | [IssueConfig](Issues.md)<br>Build health issue configurations
-`enableIssueRootCauseAnalysisNotifications` | `boolean`<br>Whether to allow notifications for Root Cause Analysis to be sent
-`issueFindSuspectsMaxSuccessDays` | `integer`<br>The maximum days since node success to scan for suspects
-`issueFindSuspectsMaxScoringSeconds` | `integer`<br>The maximum number of seconds to spend scording suspects
+`stores` | [TelemetryStoreConfig](Telemetry.md)`[]`<br>Metrics to aggregate on the Horde server
 
-## PerforceCluster
+## ComputeConfig
 
-Information about a cluster of Perforce servers.
+Configuration for the compute system
 
 Name | Description
 ---- | -----------
-`name` | `string`<br>Name of the cluster
-`serviceAccount` | `string`<br>Username for Horde to log in to this server. Will use the first account specified below if not overridden.
-`canImpersonate` | `boolean`<br>Whether the service account can impersonate other users
-`supportsPartitionedWorkspaces` | `boolean`<br>Whether to use partitioned workspaces on this server
-`servers` | [PerforceServer](#perforceserver)`[]`<br>List of servers
-`credentials` | [PerforceCredentials](#perforcecredentials)`[]`<br>List of server credentials
-`autoSdk` | [AutoSdkWorkspace](#autosdkworkspace)`[]`<br>List of autosdk streams
+`acl` | [AclConfig](#aclconfig)<br>Inherited root acl
+`versionEnum` | [ConfigVersion](#configversion-enum)<br>Config version number
+`rates` | [AgentRateConfig](#agentrateconfig)`[]`<br>List of costs of a particular agent type
+`clusters` | [ComputeClusterConfig](#computeclusterconfig)`[]`<br>List of compute profiles
+`pools` | [PoolConfig](#poolconfig)`[]`<br>List of pools
+`software` | [AgentSoftwareConfig](#agentsoftwareconfig)`[]`<br>List of costs of a particular agent type
+`networks` | [NetworkConfig](#networkconfig)`[]`<br>List of networks
 
-## PerforceServer
-
-Information about an individual Perforce server
-
-Name | Description
----- | -----------
-`serverAndPort` | `string`<br>The server and port. The server may be a DNS entry with multiple records, in which case it will be actively load balanced. If "ssl:" prefix is used, ensure P4 server's fingerprint/certificate is trusted. See Horde's documentation on connecting to SSL-enabled Perforce servers.
-`healthCheck` | `boolean`<br>Whether to query the healthcheck address under each server
-`resolveDns` | `boolean`<br>Whether to resolve the DNS entries and load balance between different hosts
-`maxConformCount` | `integer`<br>Maximum number of simultaneous conforms on this server
-`condition` | `string`<br>Optional condition for a machine to be eligible to use this server
-`properties` | `string[]`<br>List of properties for an agent to be eligible to use this server
-
-## PerforceCredentials
-
-Credentials for a Perforce user
-
-Name | Description
----- | -----------
-`userName` | `string`<br>The username
-`password` | `string`<br>Password for the user
-`ticket` | `string`<br>Login ticket for the user (will be used instead of password if set)
-
-## AutoSdkWorkspace
-
-Path to a platform and stream to use for syncing AutoSDK
-
-Name | Description
----- | -----------
-`name` | `string`<br>Name of this workspace
-`properties` | `string[]`<br>The agent properties to check (eg. "OSFamily=Windows")
-`userName` | `string`<br>Username for logging in to the server
-`stream` | `string`<br>Stream to use
-
-## DeviceConfig
-
-Configuration for devices
-
-Name | Description
----- | -----------
-`platforms` | [DevicePlatformConfig](#deviceplatformconfig)`[]`<br>List of device platforms
-`pools` | [DevicePoolConfig](#devicepoolconfig)`[]`<br>List of device pools
-
-## DevicePlatformConfig
-
-Configuration for a device platform
-
-Name | Description
----- | -----------
-`id` | `string`<br>The id for this platform
-`name` | `string`<br>Name of the platform
-`models` | `string[]`<br>A list of platform models
-`legacyNames` | `string[]`<br>Legacy names which older versions of Gauntlet may be using
-`legacyPerfSpecHighModel` | `string`<br>Model name for the high perf spec, which may be requested by Gauntlet
-
-## DevicePoolConfig
-
-Configuration for a device pool
-
-Name | Description
----- | -----------
-`id` | `string`<br>The id for this platform
-`name` | `string`<br>The name of the pool
-`poolType` | [DevicePoolType](#devicepooltype-enum)<br>The type of the pool
-`projectIds` | `string[]`<br>List of project ids associated with pool
-
-## DevicePoolType (Enum)
-
-The type of device pool
-
-Name | Description
----- | -----------
-`Automation` | Available to CIS jobs
-`Shared` | Shared by users with remote checking and checkouts
-
-## ArtifactTypeConfig
-
-Configuration for an artifact
-
-Name | Description
----- | -----------
-`name` | `string`<br>Legacy 'Name' property
-`type` | `string`<br>Name of the artifact type
-`acl` | [AclConfig](#aclconfig)<br>Acl for the artifact type
-`keepCount` | `integer`<br>Number of artifacts to retain
-`keepDays` | `integer`<br>Number of days to retain artifacts of this type
-`namespaceId` | `string`<br>Storage namespace to use for this artifact types
-
 ## AclConfig
 
@@ -232,28 +132,6 @@
 `excludeActions` | `string[]`<br>Actions to exclude from the inherited actions
 `extends` | `string[]`<br>Other profiles to extend from
 
-## AnalyticsConfig
-
-Config settings for analytics
-
-Name | Description
----- | -----------
-`stores` | [TelemetryStoreConfig](Telemetry.md)`[]`<br>Metrics to aggregate on the Horde server
-
-## ComputeConfig
-
-Configuration for the compute system
-
-Name | Description
----- | -----------
-`acl` | [AclConfig](#aclconfig)<br>Inherited root acl
-`versionEnum` | [ConfigVersion](#configversion-enum)<br>Config version number
-`rates` | [AgentRateConfig](#agentrateconfig)`[]`<br>List of costs of a particular agent type
-`clusters` | [ComputeClusterConfig](#computeclusterconfig)`[]`<br>List of compute profiles
-`pools` | [PoolConfig](#poolconfig)`[]`<br>List of pools
-`software` | [AgentSoftwareConfig](#agentsoftwareconfig)`[]`<br>List of costs of a particular agent type
-`networks` | [NetworkConfig](#networkconfig)`[]`<br>List of networks
-
 ## ConfigVersion (Enum)
 
 Global version number for running the server. As new features are introduced that require data migrations, this version number indicates the backwards compatibility functionality that must be enabled. When adding a new version here, also add a message to ConfigService.CreateSnapshotAsync describing the steps that need to be taken to upgrade the deployment.
@@ -435,6 +313,128 @@
 `description` | `string`<br>Human-readable description
 `computeId` | `string`<br>Compute ID for this network (used when allocating compute resources)
 
+## BuildConfig
+
+Configuration for the build plugin
+
+Name | Description
+---- | -----------
+`perforceClusters` | [PerforceCluster](#perforcecluster)`[]`<br>List of Perforce clusters
+`devices` | [DeviceConfig](#deviceconfig)<br>Device configuration
+`maxConformCount` | `integer`<br>Maximum number of conforms to run at once
+`agentShutdownIfDisabledGracePeriod` | `string`<br>Time to wait before shutting down an agent that has been disabled Used if no value is set on the actual pool.
+`artifactTypes` | [ArtifactTypeConfig](#artifacttypeconfig)`[]`<br>Configuration for different artifact types
+`projects` | [ProjectConfig](Projects.md)`[]`<br>List of projects
+`enableConformTasks` | `boolean`<br>Whether to allow conform tasks to run
+`enableJobTaskConformCheck` | `boolean`<br>Whether to check if a conform is required on an agent before assigning a lease for a job task source.
+`issueFixedTag` | `string`<br>Commit tag to use for marking issues as fixed
+`issues` | [IssueConfig](Issues.md)<br>Build health issue configurations
+`enableIssueRootCauseAnalysisNotifications` | `boolean`<br>Whether to allow notifications for Root Cause Analysis to be sent
+`issueFindSuspectsMaxSuccessDays` | `integer`<br>The maximum days since node success to scan for suspects
+`issueFindSuspectsMaxScoringSeconds` | `integer`<br>The maximum number of seconds to spend scording suspects
+
+## PerforceCluster
+
+Information about a cluster of Perforce servers.
+
+Name | Description
+---- | -----------
+`name` | `string`<br>Name of the cluster
+`serviceAccount` | `string`<br>Username for Horde to log in to this server. Will use the first account specified below if not overridden.
+`canImpersonate` | `boolean`<br>Whether the service account can impersonate other users
+`supportsPartitionedWorkspaces` | `boolean`<br>Whether to use partitioned workspaces on this server
+`servers` | [PerforceServer](#perforceserver)`[]`<br>List of servers
+`credentials` | [PerforceCredentials](#perforcecredentials)`[]`<br>List of server credentials
+`autoSdk` | [AutoSdkWorkspace](#autosdkworkspace)`[]`<br>List of autosdk streams
+
+## PerforceServer
+
+Information about an individual Perforce server
+
+Name | Description
+---- | -----------
+`serverAndPort` | `string`<br>The server and port. The server may be a DNS entry with multiple records, in which case it will be actively load balanced. If "ssl:" prefix is used, ensure P4 server's fingerprint/certificate is trusted. See Horde's documentation on connecting to SSL-enabled Perforce servers.
+`healthCheck` | `boolean`<br>Whether to query the healthcheck address under each server
+`resolveDns` | `boolean`<br>Whether to resolve the DNS entries and load balance between different hosts
+`maxConformCount` | `integer`<br>Maximum number of simultaneous conforms on this server
+`condition` | `string`<br>Optional condition for a machine to be eligible to use this server
+`properties` | `string[]`<br>List of properties for an agent to be eligible to use this server
+
+## PerforceCredentials
+
+Credentials for a Perforce user
+
+Name | Description
+---- | -----------
+`userName` | `string`<br>The username
+`password` | `string`<br>Password for the user
+`ticket` | `string`<br>Login ticket for the user (will be used instead of password if set)
+
+## AutoSdkWorkspace
+
+Path to a platform and stream to use for syncing AutoSDK
+
+Name | Description
+---- | -----------
+`name` | `string`<br>Name of this workspace
+`properties` | `string[]`<br>The agent properties to check (eg. "OSFamily=Windows")
+`userName` | `string`<br>Username for logging in to the server
+`stream` | `string`<br>Stream to use
+
+## DeviceConfig
+
+Configuration for devices
+
+Name | Description
+---- | -----------
+`platforms` | [DevicePlatformConfig](#deviceplatformconfig)`[]`<br>List of device platforms
+`pools` | [DevicePoolConfig](#devicepoolconfig)`[]`<br>List of device pools
+
+## DevicePlatformConfig
+
+Configuration for a device platform
+
+Name | Description
+---- | -----------
+`id` | `string`<br>The id for this platform
+`name` | `string`<br>Name of the platform
+`models` | `string[]`<br>A list of platform models
+`legacyNames` | `string[]`<br>Legacy names which older versions of Gauntlet may be using
+`legacyPerfSpecHighModel` | `string`<br>Model name for the high perf spec, which may be requested by Gauntlet
+
+## DevicePoolConfig
+
+Configuration for a device pool
+
+Name | Description
+---- | -----------
+`id` | `string`<br>The id for this platform
+`name` | `string`<br>The name of the pool
+`poolType` | [DevicePoolType](#devicepooltype-enum)<br>The type of the pool
+`projectIds` | `string[]`<br>List of project ids associated with pool
+
+## DevicePoolType (Enum)
+
+The type of device pool
+
+Name | Description
+---- | -----------
+`Automation` | Available to CIS jobs
+`Shared` | Shared by users with remote checking and checkouts
+
+## ArtifactTypeConfig
+
+Configuration for an artifact
+
+Name | Description
+---- | -----------
+`name` | `string`<br>Legacy 'Name' property
+`type` | `string`<br>Name of the artifact type
+`acl` | [AclConfig](#aclconfig)<br>Acl for the artifact type
+`keepCount` | `integer`<br>Number of artifacts to retain
+`keepDays` | `integer`<br>Number of days to retain artifacts of this type
+`namespaceId` | `string`<br>Storage namespace to use for this artifact types
+
 ## ExamplePluginConfig
 
 Configuration for the ExamplePlugin plugin
@@ -451,77 +451,5 @@
 ---- | -----------
 `notifications` | [NotificationConfig](Notification.md)`[]`<br>Notifications to be sent on the Horde server
 
-## StorageConfig
-
-Configuration for storage
-
-Name | Description
----- | -----------
-`enableGc` | `boolean`<br>Whether to enable garbage collection
-`enableGcVerification` | `boolean`<br>Whether to enable garbage collection in verification mode (nothing deleted, just logging on access to deleted blobs)
-`backends` | [BackendConfig](#backendconfig)`[]`<br>List of storage backends
-`namespaces` | [NamespaceConfig](#namespaceconfig)`[]`<br>List of namespaces for storage
-
-## BackendConfig
-
-Common settings object for different providers
-
-Name | Description
----- | -----------
-`id` | `string`<br>The storage backend ID
-`base` | `string`<br>Base backend to copy default settings from
-`secondary` | `string`<br>Specifies another backend to read from if an object is not found in this one. Can be used when migrating data from one backend to another.
-`type` | [StorageBackendType](#storagebackendtype-enum)<br>
-`baseDir` | `string`<br>
-`awsBucketName` | `string`<br>Name of the bucket to use
-`awsBucketPath` | `string`<br>Base path within the bucket
-`awsCredentials` | [AwsCredentialsType](#awscredentialstype-enum)<br>Type of credentials to use
-`awsRole` | `string`<br>ARN of a role to assume
-`awsProfile` | `string`<br>The AWS profile to read credentials form
-`awsRegion` | `string`<br>Region to connect to
-`azureConnectionString` | `string`<br>Connection string for Azure
-`azureContainerName` | `string`<br>Name of the container
-`relayServer` | `string`<br>
-`relayToken` | `string`<br>
-`gcsBucketName` | `string`<br>Name of the GCS bucket to use
-`gcsBucketPath` | `string`<br>Base path within the bucket
-
-## StorageBackendType (Enum)
-
-Types of storage backend to use
-
-Name | Description
----- | -----------
-`FileSystem` | Local filesystem
-`Aws` | AWS S3
-`Azure` | Azure blob store
-`Gcs` | Google Cloud Storage
-`Memory` | In-memory only (for testing)
-
-## AwsCredentialsType (Enum)
-
-Credentials to use for AWS
-
-Name | Description
----- | -----------
-`Default` | Use default credentials from the AWS SDK
-`Profile` | Read credentials from the  profile in the AWS config file
-`AssumeRole` | Assume a particular role. Should specify ARN in
-`AssumeRoleWebIdentity` | Assume a particular role using the current environment variables.
-
-## NamespaceConfig
-
-Configuration of a particular namespace
-
-Name | Description
----- | -----------
-`id` | `string`<br>Identifier for this namespace
-`backend` | `string`<br>Backend to use for this namespace
-`prefix` | `string`<br>Prefix for items within this namespace
-`gcFrequencyHrs` | `number`<br>How frequently to run garbage collection, in hours.
-`gcDelayHrs` | `number`<br>How long to keep newly uploaded orphanned objects before allowing them to be deleted, in hours.
-`enableAliases` | `boolean`<br>Support querying exports by their aliases
-`acl` | [AclConfig](#aclconfig)<br>Access list for this namespace
-
 ## SecretsConfig
 
@@ -605,24 +533,77 @@
 ---- | -----------
 `vaultUri` | `string`<br>The Uri to the Azure Key Vault.
 
-## SymbolsConfig
+## StorageConfig
 
-Configuration for the tools system
+Configuration for storage
 
 Name | Description
 ---- | -----------
-`stores` | [SymbolStoreConfig](#symbolstoreconfig)`[]`<br>List of symbol stores
+`enableGc` | `boolean`<br>Whether to enable garbage collection
+`enableGcVerification` | `boolean`<br>Whether to enable garbage collection in verification mode (nothing deleted, just logging on access to deleted blobs)
+`backends` | [BackendConfig](#backendconfig)`[]`<br>List of storage backends
+`namespaces` | [NamespaceConfig](#namespaceconfig)`[]`<br>List of namespaces for storage
 
-## SymbolStoreConfig
+## BackendConfig
 
-Configuration for a symbol store
+Common settings object for different providers
 
 Name | Description
 ---- | -----------
-`id` | `string`<br>Identifier for this store
-`namespaceId` | `string`<br>Configuration for the symbol store backend
-`public` | `boolean`<br>Whether to make this store available without auth
-`acl` | [AclConfig](#aclconfig)<br>Access to the symbol store
+`id` | `string`<br>The storage backend ID
+`base` | `string`<br>Base backend to copy default settings from
+`secondary` | `string`<br>Specifies another backend to read from if an object is not found in this one. Can be used when migrating data from one backend to another.
+`type` | [StorageBackendType](#storagebackendtype-enum)<br>
+`baseDir` | `string`<br>
+`awsBucketName` | `string`<br>Name of the bucket to use
+`awsBucketPath` | `string`<br>Base path within the bucket
+`awsCredentials` | [AwsCredentialsType](#awscredentialstype-enum)<br>Type of credentials to use
+`awsRole` | `string`<br>ARN of a role to assume
+`awsProfile` | `string`<br>The AWS profile to read credentials form
+`awsRegion` | `string`<br>Region to connect to
+`azureConnectionString` | `string`<br>Connection string for Azure
+`azureContainerName` | `string`<br>Name of the container
+`relayServer` | `string`<br>
+`relayToken` | `string`<br>
+`gcsBucketName` | `string`<br>Name of the GCS bucket to use
+`gcsBucketPath` | `string`<br>Base path within the bucket
 
+## StorageBackendType (Enum)
+
+Types of storage backend to use
+
+Name | Description
+---- | -----------
+`FileSystem` | Local filesystem
+`Aws` | AWS S3
+`Azure` | Azure blob store
+`Gcs` | Google Cloud Storage
+`Memory` | In-memory only (for testing)
+
+## AwsCredentialsType (Enum)
+
+Credentials to use for AWS
+
+Name | Description
+---- | -----------
+`Default` | Use default credentials from the AWS SDK
+`Profile` | Read credentials from the  profile in the AWS config file
+`AssumeRole` | Assume a particular role. Should specify ARN in
+`AssumeRoleWebIdentity` | Assume a particular role using the current environment variables.
+
+## NamespaceConfig
+
+Configuration of a particular namespace
+
+Name | Description
+---- | -----------
+`id` | `string`<br>Identifier for this namespace
+`backend` | `string`<br>Backend to use for this namespace
+`prefix` | `string`<br>Prefix for items within this namespace
+`gcFrequencyHrs` | `number`<br>How frequently to run garbage collection, in hours.
+`gcDelayHrs` | `number`<br>How long to keep newly uploaded orphanned objects before allowing them to be deleted, in hours.
+`enableAliases` | `boolean`<br>Support querying exports by their aliases
+`acl` | [AclConfig](#aclconfig)<br>Access list for this namespace
+
 ## ToolsConfig
 
@@ -651,3 +632,22 @@
 `metadata` | `string` `->` `string`<br>Metadata for this tool
 `namespaceId` | `string`<br>Default namespace for new deployments of this tool
 `acl` | [AclConfig](#aclconfig)<br>Permissions for the tool
+
+## SymbolsConfig
+
+Configuration for the tools system
+
+Name | Description
+---- | -----------
+`stores` | [SymbolStoreConfig](#symbolstoreconfig)`[]`<br>List of symbol stores
+
+## SymbolStoreConfig
+
+Configuration for a symbol store
+
+Name | Description
+---- | -----------
+`id` | `string`<br>Identifier for this store
+`namespaceId` | `string`<br>Configuration for the symbol store backend
+`public` | `boolean`<br>Whether to make this store available without auth
+`acl` | [AclConfig](#aclconfig)<br>Access to the symbol store
