Horde: Fix potential issue with CPU time calculations in agent

Also wrap lease execution to ensure any exceptions are immediately surfaced.

diff --git a/Engine/Source/Programs/Horde/HordeAgent/Leases/LeaseHandler.cs b/Engine/Source/Programs/Horde/HordeAgent/Leases/LeaseHandler.cs
--- a/Engine/Source/Programs/Horde/HordeAgent/Leases/LeaseHandler.cs
+++ b/Engine/Source/Programs/Horde/HordeAgent/Leases/LeaseHandler.cs
@@ -87,7 +87,18 @@ namespace HordeAgent.Leases
 		/// </summary>
 		public void Start(ISession session, ISystemMetrics systemMetrics, Tracer tracer, ILogger logger, LeaseLoggerFactory leaseLoggerFactory)
 		{
-			Result = Task.Run(() => HandleLeaseAsync(session, systemMetrics, tracer, logger, leaseLoggerFactory));
+			Result = Task.Run(async () =>
+			{
+				try
+				{
+					return await HandleLeaseAsync(session, systemMetrics, tracer, logger, leaseLoggerFactory);
+				}
+				catch (Exception e)
+				{
+					logger.LogError(e, "Unhandled exception while executing lease: {Message} ", e.Message);
+					throw;
+				}
+			});
 		}
 
 		/// <summary>
@@ -149,37 +160,8 @@ namespace HordeAgent.Leases
 				newRpcLease.OutcomeReasons.AddRange(RpcLease.ToProtobufOutcomeReasons(reason));
 				newRpcLease.Output = (result.Output != null) ? ByteString.CopyFrom(result.Output) : ByteString.Empty;
 			}
-			
-			newRpcLease.Metrics = new RpcLeaseMetrics
-			{
-				SetupTimeMs = 0, // TODO: Implement
-				TeardownTimeMs = 0, // TODO: Implement
-				GlobalCpuUtilization = 0.0f,
-				CpuCount = Environment.ProcessorCount,
-			};
-			
-			if (cpuTimeStart != null && cpuTimeEnd != null)
-			{
-				CpuTimeMetrics cpuTimeSpent = cpuTimeEnd.Subtract(cpuTimeStart);
-				TimeSpan totalCpuTimeSpent = cpuTimeSpent.User + cpuTimeSpent.System;
-				TimeSpan totalCpuTimeAvailable = Environment.ProcessorCount * leaseDuration.Elapsed;
-				
-				// Prevent the unlikely divide by zero below
-				totalCpuTimeAvailable = totalCpuTimeAvailable == TimeSpan.Zero ? TimeSpan.FromMilliseconds(1) : totalCpuTimeAvailable;
-				
-				newRpcLease.Metrics.GlobalCpuUtilization = totalCpuTimeSpent / totalCpuTimeAvailable;
-				newRpcLease.Metrics.GlobalCpuTime = new RpcLeaseMetricCpuTime
-				{
-					UserMs = Convert.ToUInt32(cpuTimeSpent.User.TotalMilliseconds),
-					SystemMs = Convert.ToUInt32(cpuTimeSpent.System.TotalMilliseconds),
-					IdleMs = Convert.ToUInt32(cpuTimeSpent.Idle.TotalMilliseconds),
-				};
-				span.SetAttribute("horde.lease.cpu.time_spent_ms", (uint)totalCpuTimeSpent.TotalMilliseconds);
-				span.SetAttribute("horde.lease.cpu.time_available_ms", (uint)totalCpuTimeAvailable.TotalMilliseconds);
-				span.SetAttribute("horde.lease.cpu.utilization", newRpcLease.Metrics.GlobalCpuUtilization);
-				span.SetAttribute("horde.lease.cpu.num_cores", newRpcLease.Metrics.CpuCount);
-			}
-			
+
+			newRpcLease.Metrics = CalculateLeaseMetrics(cpuTimeStart, cpuTimeEnd, leaseDuration.Elapsed, span, logger);
 			RpcLease = newRpcLease;
 
 			logger.LogInformation("Transitioning lease {LeaseId} to {State}, outcome={Outcome}", Id, RpcLease.State, RpcLease.Outcome);
@@ -276,6 +258,47 @@ namespace HordeAgent.Leases
 				return await RunProcessAsync("dotnet", allArguments, environment, logger, cancellationToken);
 			}
 		}
+
+		private static RpcLeaseMetrics CalculateLeaseMetrics(CpuTimeMetrics? cpuTimeStart, CpuTimeMetrics? cpuTimeEnd, TimeSpan leaseDuration, TelemetrySpan span, ILogger logger)
+		{
+			RpcLeaseMetrics metrics = new()
+			{
+				SetupTimeMs = 0, // TODO: Implement
+				TeardownTimeMs = 0, // TODO: Implement
+				GlobalCpuUtilization = 0.0f,
+				CpuCount = Environment.ProcessorCount,
+			};
+			try
+			{
+				if (cpuTimeStart != null && cpuTimeEnd != null)
+				{
+					CpuTimeMetrics cpuTimeSpent = cpuTimeEnd.Subtract(cpuTimeStart);
+					TimeSpan totalCpuTimeSpent = cpuTimeSpent.User + cpuTimeSpent.System;
+					TimeSpan totalCpuTimeAvailable = Environment.ProcessorCount * leaseDuration;
+				
+					// Prevent the unlikely divide by zero below
+					totalCpuTimeAvailable = totalCpuTimeAvailable == TimeSpan.Zero ? TimeSpan.FromMilliseconds(1) : totalCpuTimeAvailable;
+				
+					metrics.GlobalCpuUtilization = totalCpuTimeSpent / totalCpuTimeAvailable;
+					metrics.GlobalCpuTime = new RpcLeaseMetricCpuTime
+					{
+						UserMs = (uint)Math.Min(uint.MaxValue, cpuTimeSpent.User.TotalMilliseconds),
+						SystemMs = (uint)Math.Min(uint.MaxValue, cpuTimeSpent.System.TotalMilliseconds),
+						IdleMs = (uint)Math.Min(uint.MaxValue, cpuTimeSpent.Idle.TotalMilliseconds),
+					};
+					span.SetAttribute("horde.lease.cpu.time_spent_ms", (uint)totalCpuTimeSpent.TotalMilliseconds);
+					span.SetAttribute("horde.lease.cpu.time_available_ms", (uint)totalCpuTimeAvailable.TotalMilliseconds);
+					span.SetAttribute("horde.lease.cpu.utilization", metrics.GlobalCpuUtilization);
+					span.SetAttribute("horde.lease.cpu.num_cores", metrics.CpuCount);
+				}
+			}
+			catch (Exception e)
+			{
+				logger.LogWarning(e, "Failed to calculate lease metrics");
+				
+			}
+			return metrics;
+		}
 	}
 
 	/// <summary>
