[TestData] Fix missing cancellation token in controller interface

#rnx
#Synced-CL 49441698

diff --git a/Engine/Source/Programs/Horde/Plugins/Build/HordeServer.Build/Jobs/TestData/TestDataController.cs b/Engine/Source/Programs/Horde/Plugins/Build/HordeServer.Build/Jobs/TestData/TestDataController.cs
--- a/Engine/Source/Programs/Horde/Plugins/Build/HordeServer.Build/Jobs/TestData/TestDataController.cs
+++ b/Engine/Source/Programs/Horde/Plugins/Build/HordeServer.Build/Jobs/TestData/TestDataController.cs
@@ -542,7 +542,7 @@ namespace HordeServer.Jobs.TestData
 		[HttpPost]
 		[Route("/api/v3/testdata/streams")]
 		[ProducesResponseType(typeof(List<GetTestSessionStreamResponse>), 200)]
-		public async Task<ActionResult<List<GetTestSessionStreamResponse>>> GetTestStreamsV2Async([FromBody] GetStreamTestsRequest request)
+		public async Task<ActionResult<List<GetTestSessionStreamResponse>>> GetTestStreamsV2Async([FromBody] GetStreamTestsRequest request, CancellationToken cancellationToken = default)
 		{
 			StreamId[] streamIdValues = request.StreamIds.Select(x => new StreamId(x)).ToArray();
 
@@ -563,7 +563,7 @@ namespace HordeServer.Jobs.TestData
 				return responses;
 			}
 
-			IReadOnlyList<ITestSessionStream> streams = await _testDataService.FindTestSessionStreamsV2Async(queryStreams.ToArray());
+			IReadOnlyList<ITestSessionStream> streams = await _testDataService.FindTestSessionStreamsV2Async(queryStreams.ToArray(), cancellationToken);
 
 			// flatten requested streams to single service queries		
 			HashSet<TestId> testIds = new HashSet<TestId>();
@@ -580,21 +580,21 @@ namespace HordeServer.Jobs.TestData
 			IReadOnlyList<ITestNameRef> tests = new List<ITestNameRef>();
 			if (testIds.Count > 0)
 			{
-				tests = await _testDataService.FindTestNameRefsV2Async(testIds.ToArray());
+				tests = await _testDataService.FindTestNameRefsV2Async(testIds.ToArray(), null, cancellationToken);
 			}
 
 			// Gather all relevant metadata
 			IReadOnlyList<ITestMetaRef> metadata = new List<ITestMetaRef>();
 			if (metaIds.Count > 0)
 			{
-				metadata = await _testDataService.FindTestMetaV2Async(metaIds.ToArray());
+				metadata = await _testDataService.FindTestMetaV2Async(metaIds.ToArray(), null, cancellationToken);
 			}
 
 			// Gather all relevant tags
 			IReadOnlyList<ITestTagRef> tags = new List<ITestTagRef>();
 			if (tagIds.Count > 0)
 			{
-				tags = await _testDataService.FindTestTagsV2Async(tagIds.ToArray());
+				tags = await _testDataService.FindTestTagsV2Async(tagIds.ToArray(), null, cancellationToken);
 			}
 
 			// generate individual stream responses
@@ -641,11 +641,11 @@ namespace HordeServer.Jobs.TestData
 		[HttpPost]
 		[Route("/api/v3/testdata/tests")]
 		[ProducesResponseType(typeof(List<GetTestNameResponse>), 200)]
-		public async Task<ActionResult<List<GetTestNameResponse>>> GetTestNameRefV2Async([FromBody] GetTestsRequest request)
+		public async Task<ActionResult<List<GetTestNameResponse>>> GetTestNameRefV2Async([FromBody] GetTestsRequest request, CancellationToken cancellationToken = default)
 		{
 			HashSet<string> testIds = new HashSet<string>(request.TestIds);
 
-			IReadOnlyList<ITestNameRef> testValues = await _testDataService.FindTestNameRefsV2Async(testIds.Select(x => TestId.Parse(x)).ToArray());
+			IReadOnlyList<ITestNameRef> testValues = await _testDataService.FindTestNameRefsV2Async(testIds.Select(x => TestId.Parse(x)).ToArray(), null, cancellationToken);
 
 			return testValues.Select(x => new GetTestNameResponse
 			{
@@ -667,11 +667,11 @@ namespace HordeServer.Jobs.TestData
 		[HttpPost]
 		[Route("/api/v3/testdata/metadata")]
 		[ProducesResponseType(typeof(List<GetTestMetadataResponse>), 200)]
-		public async Task<ActionResult<List<GetTestMetadataResponse>>> GetTestMetadataV2Async([FromBody] GetMetadataRequest request)
+		public async Task<ActionResult<List<GetTestMetadataResponse>>> GetTestMetadataV2Async([FromBody] GetMetadataRequest request, CancellationToken cancellationToken = default)
 		{
 			TestMetaId[]? metaIds = request.MetadataIds?.Select(x => TestMetaId.Parse(x)).ToHashSet().ToArray();
 
-			IReadOnlyList<ITestMetaRef> metadataValues = await _testDataService.FindTestMetaV2Async(metaIds, request.Entries);
+			IReadOnlyList<ITestMetaRef> metadataValues = await _testDataService.FindTestMetaV2Async(metaIds, request.Entries, cancellationToken);
 
 			return metadataValues.Select(x => new GetTestMetadataResponse
 			{
@@ -688,11 +688,11 @@ namespace HordeServer.Jobs.TestData
 		[HttpPost]
 		[Route("/api/v3/testdata/tags")]
 		[ProducesResponseType(typeof(List<GetTestTagResponse>), 200)]
-		public async Task<ActionResult<List<GetTestTagResponse>>> GetTestTagRefV2Async([FromBody] GetTestTagRequest request)
+		public async Task<ActionResult<List<GetTestTagResponse>>> GetTestTagRefV2Async([FromBody] GetTestTagRequest request, CancellationToken cancellationToken = default)
 		{
 			TestTagId[]? tagIds = request.TagIds?.Select(x => TestTagId.Parse(x)).ToHashSet().ToArray();
 
-			IReadOnlyList<ITestTagRef> testValues = await _testDataService.FindTestTagsV2Async(tagIds, request.TagNames?.ToArray());
+			IReadOnlyList<ITestTagRef> testValues = await _testDataService.FindTestTagsV2Async(tagIds, request.TagNames?.ToArray(), cancellationToken);
 
 			return testValues.Select(x => new GetTestTagResponse
 			{
@@ -709,13 +709,13 @@ namespace HordeServer.Jobs.TestData
 		[HttpPost]
 		[Route("/api/v3/testdata/recipes")]
 		[ProducesResponseType(typeof(List<GetTestRecipeResponse>), 200)]
-		public async Task<ActionResult<List<GetTestRecipeResponse>>> GetTestRecipeRefV2Async([FromBody] GetTestRecipesRequest request)
+		public async Task<ActionResult<List<GetTestRecipeResponse>>> GetTestRecipeRefV2Async([FromBody] GetTestRecipesRequest request, CancellationToken cancellationToken = default)
 		{
 			TestRecipeId[]? recipeIds = request.RecipeIds?.Select(x => TestRecipeId.Parse(x)).ToHashSet().ToArray();
 			TestId[]? testIds = request.TestIds?.Select(x => TestId.Parse(x)).ToHashSet().ToArray();
 			StreamId[]? streamIds = request.StreamIds?.Select(x => new StreamId(x)).ToHashSet().ToArray();
 
-			IReadOnlyList<ITestRecipeRef> testValues = await _testDataService.FindTestRecipesV2Async(recipeIds, testIds, streamIds);
+			IReadOnlyList<ITestRecipeRef> testValues = await _testDataService.FindTestRecipesV2Async(recipeIds, testIds, streamIds, cancellationToken);
 
 			return testValues.Where(
 				x => _buildConfig.Value.TryGetStream(x.StreamId, out StreamConfig? streamConfig) && streamConfig.Authorize(JobAclAction.ViewJob, User)
@@ -742,12 +742,12 @@ namespace HordeServer.Jobs.TestData
 		[HttpPost]
 		[Route("/api/v3/testdata/phases")]
 		[ProducesResponseType(typeof(List<GetTestPhaseResponse>), 200)]
-		public async Task<ActionResult<List<GetTestPhaseResponse>>> GetTestPhaseRefFromTestV2Async([FromBody] GetTestPhasesRequest request)
+		public async Task<ActionResult<List<GetTestPhaseResponse>>> GetTestPhaseRefFromTestV2Async([FromBody] GetTestPhasesRequest request, CancellationToken cancellationToken = default)
 		{
 			string[]? testKeys = request.TestKeys?.ToHashSet().ToArray();
 			TestId[]? testIds = request.TestIds?.Select(x => TestId.Parse(x)).ToHashSet().ToArray();
-			IReadOnlyList<ITestNameRef> testNameRefs = await _testDataService.FindTestNameRefsV2Async(testIds, testKeys);
-			IReadOnlyList<ITestPhaseRef> phaseValues = await _testDataService.FindTestPhasesV2Async(testNameRefs.Select(t => t.Id).ToArray());
+			IReadOnlyList<ITestNameRef> testNameRefs = await _testDataService.FindTestNameRefsV2Async(testIds, testKeys, cancellationToken);
+			IReadOnlyList<ITestPhaseRef> phaseValues = await _testDataService.FindTestPhasesV2Async(testNameRefs.Select(t => t.Id).ToArray(), null, cancellationToken);
 
 			List<GetTestPhaseResponse> response = new List<GetTestPhaseResponse>();
 
@@ -788,7 +788,8 @@ namespace HordeServer.Jobs.TestData
 			[FromQuery] DateTimeOffset? minCreateTime = null,
 			[FromQuery] DateTimeOffset? maxCreateTime = null,
 			[FromQuery] CommitId? minChange = null,
-			[FromQuery] CommitId? maxChange = null)
+			[FromQuery] CommitId? maxChange = null,
+			CancellationToken cancellationToken = default)
 		{
 			StreamId[] streamIdValues = request.StreamIds.Select(x => new StreamId(x)).ToArray();
 
@@ -812,7 +813,7 @@ namespace HordeServer.Jobs.TestData
 			TestMetaId[]? queryMetaIds = request.MetaIds?.Select(x => TestMetaId.Parse(x)).ToArray();
 			TestId[]? queryTestIds = request.TestIds?.Select(x => TestId.Parse(x)).ToArray();
 
-			IReadOnlyList<ITestSession> sessions = await _testDataService.FindTestSessionsV2Async(queryStreams.ToArray(), queryTestIds, queryMetaIds, minCreateTime?.UtcDateTime, maxCreateTime?.UtcDateTime, minChange, maxChange);
+			IReadOnlyList<ITestSession> sessions = await _testDataService.FindTestSessionsV2Async(queryStreams.ToArray(), queryTestIds, queryMetaIds, minCreateTime?.UtcDateTime, maxCreateTime?.UtcDateTime, minChange, maxChange, cancellationToken);
 			foreach (ITestSession testData in sessions)
 			{
 				responses.Add(new GetTestSessionResponse
@@ -864,7 +865,8 @@ namespace HordeServer.Jobs.TestData
 			[FromQuery] DateTimeOffset? minCreateTime = null,
 			[FromQuery] DateTimeOffset? maxCreateTime = null,
 			[FromQuery] CommitId? minChange = null,
-			[FromQuery] CommitId? maxChange = null)
+			[FromQuery] CommitId? maxChange = null,
+			CancellationToken cancellationToken = default)
 		{
 			StreamId[] streamIdValues = request.StreamIds.Select(x => new StreamId(x)).ToArray();
 
@@ -887,7 +889,7 @@ namespace HordeServer.Jobs.TestData
 
 			TestPhaseId[] queryPhaseIds = request.PhaseIds.Select(x => TestPhaseId.Parse(x)).ToArray();
 
-			IReadOnlyList<ITestPhaseSession> sessions = await _testDataService.FindTestPhaseSessionsV2Async(queryStreams.ToArray(), queryPhaseIds, minCreateTime?.UtcDateTime, maxCreateTime?.UtcDateTime, minChange, maxChange);
+			IReadOnlyList<ITestPhaseSession> sessions = await _testDataService.FindTestPhaseSessionsV2Async(queryStreams.ToArray(), queryPhaseIds, minCreateTime?.UtcDateTime, maxCreateTime?.UtcDateTime, minChange, maxChange, cancellationToken);
 			foreach (ITestPhaseSession testData in sessions)
 			{
 				responses.Add(new GetTestPhaseSessionResponse
@@ -922,7 +924,7 @@ namespace HordeServer.Jobs.TestData
 		[HttpPost]
 		[Route("/api/v3/testdata/audit/post")]
 		[ProducesResponseType(typeof(GetTestAuditResponse), 200)]
-		public async Task<ActionResult<GetTestAuditResponse>> AddTestAuditV2Async([FromBody] AddTestAuditRequest request)
+		public async Task<ActionResult<GetTestAuditResponse>> AddTestAuditV2Async([FromBody] AddTestAuditRequest request, CancellationToken cancellationToken = default)
 		{
 			GetTestAuditResponse responses = new GetTestAuditResponse();
 
@@ -957,7 +959,7 @@ namespace HordeServer.Jobs.TestData
 				Notes = request.Notes
 			};
 
-			ITestAudit audit = await _testDataService.AddTestAuditAsync(userId!.Value, TestId.Parse(request.TestId), options);
+			ITestAudit audit = await _testDataService.AddTestAuditAsync(userId!.Value, TestId.Parse(request.TestId), options, cancellationToken);
 
 			return GetNewTestAutditResponse(audit);
 		}
@@ -973,11 +975,12 @@ namespace HordeServer.Jobs.TestData
 		[ProducesResponseType(typeof(List<GetTestAuditResponse>), 200)]
 		public async Task<ActionResult<List<GetTestAuditResponse>>> GetTestAuditsV2Async(
 			[FromBody] GetTestAuditRequest request,
-			[FromQuery] DateTimeOffset? minCreateTime = null)
+			[FromQuery] DateTimeOffset? minCreateTime = null,
+			CancellationToken cancellationToken = default)
 		{
 			TestId[] testIds = request.TestIds.Select(x => TestId.Parse(x)).ToHashSet().ToArray();
 
-			IReadOnlyList<ITestAudit> audits = await _testDataService.FindTestAuditsV2Async(testIds, minCreateTime?.UtcDateTime);
+			IReadOnlyList<ITestAudit> audits = await _testDataService.FindTestAuditsV2Async(testIds, minCreateTime?.UtcDateTime, cancellationToken);
 
 			return audits.Select(a => GetNewTestAutditResponse(a)).ToList();
 		}
@@ -1016,7 +1019,7 @@ namespace HordeServer.Jobs.TestData
 		/// <returns>Information about the requested test audit</returns>
 		[HttpGet]
 		[Route("/api/v3/testdata/audit/log/{testId}")]
-		public async Task<ActionResult> GetAuditLogHistoryAsync(string testId, [FromQuery] DateTime? minTime = null, [FromQuery] DateTime? maxTime = null, [FromQuery] int index = 0, [FromQuery] int count = 50)
+		public async Task<ActionResult> GetAuditLogHistoryAsync(string testId, [FromQuery] DateTime? minTime = null, [FromQuery] DateTime? maxTime = null, [FromQuery] int index = 0, [FromQuery] int count = 50, CancellationToken cancellationToken = default)
 		{
 			TestId test;
 			if (String.IsNullOrEmpty(testId) || !TestId.TryParse(testId, out test))
@@ -1026,8 +1029,8 @@ namespace HordeServer.Jobs.TestData
 
 			Response.ContentType = "application/json";
 			Response.StatusCode = 200;
-			await Response.StartAsync();
-			await _testDataService.GetTestAuditLogger(test).FindAsync(Response.BodyWriter, minTime, maxTime, index, count);
+			await Response.StartAsync(cancellationToken);
+			await _testDataService.GetTestAuditLogger(test).FindAsync(Response.BodyWriter, minTime, maxTime, index, count, cancellationToken);
 			await Response.CompleteAsync();
 
 			return Empty;
@@ -1041,7 +1044,7 @@ namespace HordeServer.Jobs.TestData
 		[HttpPost]
 		[Route("/api/v3/testdata/health/reports")]
 		[ProducesResponseType(typeof(List<GetTestHealthReportResponse>), 200)]
-		public async Task<ActionResult<List<GetTestHealthReportResponse>>> GetTestHealthReportsV2Async([FromBody] GetTestHealthReportRequest request)
+		public async Task<ActionResult<List<GetTestHealthReportResponse>>> GetTestHealthReportsV2Async([FromBody] GetTestHealthReportRequest request, CancellationToken cancellationToken = default)
 		{
 			StreamId[] streamIdValues = request.StreamIds.Select(x => new StreamId(x)).ToArray();
 
@@ -1064,7 +1067,7 @@ namespace HordeServer.Jobs.TestData
 
 			TestId[]? queryTestIds = request.TestIds?.Select(x => TestId.Parse(x)).ToArray();
 
-			IReadOnlyList<ITestHealthReport> reports = await _testDataService.FindTestHealthReportsV2Async(queryStreams.ToArray(), queryTestIds);
+			IReadOnlyList<ITestHealthReport> reports = await _testDataService.FindTestHealthReportsV2Async(queryStreams.ToArray(), queryTestIds, cancellationToken);
 
 			return reports.Select(report => new GetTestHealthReportResponse()
 			{
