Fix missing Parameter 'cancellationToken' param tag in the XML comment

#rnx
#Synced-CL 49441698

diff --git a/Engine/Source/Programs/Horde/Plugins/Build/HordeServer.Build/Jobs/TestData/TestDataController.cs b/Engine/Source/Programs/Horde/Plugins/Build/HordeServer.Build/Jobs/TestData/TestDataController.cs
--- a/Engine/Source/Programs/Horde/Plugins/Build/HordeServer.Build/Jobs/TestData/TestDataController.cs
+++ b/Engine/Source/Programs/Horde/Plugins/Build/HordeServer.Build/Jobs/TestData/TestDataController.cs
@@ -538,6 +538,7 @@ namespace HordeServer.Jobs.TestData
 		/// Get stream test data for the provided ids
 		/// </summary>
 		/// <param name="request"></param>
+		/// <param name="cancellationToken">Cancellation token for the operation</param>
 		/// <returns></returns>
 		[HttpPost]
 		[Route("/api/v3/testdata/streams")]
@@ -637,6 +638,7 @@ namespace HordeServer.Jobs.TestData
 		/// Get test names from provided ids
 		/// </summary>
 		/// <param name="request"></param>
+		/// <param name="cancellationToken">Cancellation token for the operation</param>
 		/// <returns></returns>
 		[HttpPost]
 		[Route("/api/v3/testdata/tests")]
@@ -663,6 +665,7 @@ namespace HordeServer.Jobs.TestData
 		/// Get test metadata from provided ids
 		/// </summary>
 		/// <param name="request"></param>
+		/// <param name="cancellationToken">Cancellation token for the operation</param>
 		/// <returns></returns>
 		[HttpPost]
 		[Route("/api/v3/testdata/metadata")]
@@ -684,6 +687,7 @@ namespace HordeServer.Jobs.TestData
 		/// Get test tags from provided ids
 		/// </summary>
 		/// <param name="request"></param>
+		/// <param name="cancellationToken">Cancellation token for the operation</param>
 		/// <returns></returns>
 		[HttpPost]
 		[Route("/api/v3/testdata/tags")]
@@ -705,6 +709,7 @@ namespace HordeServer.Jobs.TestData
 		/// Get test recipes from provided ids
 		/// </summary>
 		/// <param name="request"></param>
+		/// <param name="cancellationToken">Cancellation token for the operation</param>
 		/// <returns></returns>
 		[HttpPost]
 		[Route("/api/v3/testdata/recipes")]
@@ -738,6 +743,7 @@ namespace HordeServer.Jobs.TestData
 		/// Get test phases from provide test keys and refs
 		/// </summary>
 		/// <param name="request"></param>
+		/// <param name="cancellationToken">Cancellation token for the operation</param>
 		/// <returns></returns>
 		[HttpPost]
 		[Route("/api/v3/testdata/phases")]
@@ -779,6 +785,7 @@ namespace HordeServer.Jobs.TestData
 		/// <param name="maxCreateTime"></param>
 		/// <param name="minChange"></param>
 		/// <param name="maxChange"></param>
+		/// <param name="cancellationToken">Cancellation token for the operation</param>
 		/// <returns></returns>
 		[HttpPost]
 		[Route("/api/v3/testdata/session/tests")]
@@ -856,6 +863,7 @@ namespace HordeServer.Jobs.TestData
 		/// <param name="maxCreateTime"></param>
 		/// <param name="minChange"></param>
 		/// <param name="maxChange"></param>
+		/// <param name="cancellationToken">Cancellation token for the operation</param>
 		/// <returns></returns>
 		[HttpPost]
 		[Route("/api/v3/testdata/session/phases")]
@@ -920,6 +928,7 @@ namespace HordeServer.Jobs.TestData
 		/// Post a test audit
 		/// </summary>
 		/// <param name="request"></param>
+		/// <param name="cancellationToken">Cancellation token for the operation</param>
 		/// <returns></returns>
 		[HttpPost]
 		[Route("/api/v3/testdata/audit/post")]
@@ -969,6 +978,7 @@ namespace HordeServer.Jobs.TestData
 		/// </summary>
 		/// <param name="request"></param>
 		/// <param name="minCreateTime"></param>
+		/// <param name="cancellationToken">Cancellation token for the operation</param>
 		/// <returns></returns>
 		[HttpPost]
 		[Route("/api/v3/testdata/audit/history")]
@@ -1016,6 +1026,7 @@ namespace HordeServer.Jobs.TestData
 		/// <param name="maxTime">Maximum time for records to return</param>
 		/// <param name="index">Offset of the first result</param>
 		/// <param name="count">Number of records to return</param>
+		/// <param name="cancellationToken">Cancellation token for the operation</param>
 		/// <returns>Information about the requested test audit</returns>
 		[HttpGet]
 		[Route("/api/v3/testdata/audit/log/{testId}")]
@@ -1040,6 +1051,7 @@ namespace HordeServer.Jobs.TestData
 		/// Get test health reports
 		/// </summary>
 		/// <param name="request"></param>
+		/// <param name="cancellationToken">Cancellation token for the operation</param>
 		/// <returns></returns>
 		[HttpPost]
 		[Route("/api/v3/testdata/health/reports")]
