[Test Health] Fix refresh issue when navigating browser history

#rnx

#Synced-CL 49441698

diff --git a/Engine/Source/Programs/Horde/HordeDashboard/plugins/test-automation/testAutomationHealth.tsx b/Engine/Source/Programs/Horde/HordeDashboard/plugins/test-automation/testAutomationHealth.tsx
--- a/Engine/Source/Programs/Horde/HordeDashboard/plugins/test-automation/testAutomationHealth.tsx
+++ b/Engine/Source/Programs/Horde/HordeDashboard/plugins/test-automation/testAutomationHealth.tsx
@@ -960,7 +960,7 @@ export const TestAutomationHealth: React.FC<{ handler: TestDataHandler}> = obser
         });
 
         return [initItems, initColumns];
-    }, []);
+    }, [handler.searchUpdated]);
 
     const forceListRefresh = () => {
         updateColorDomains(healthListState.items);
@@ -980,7 +980,7 @@ export const TestAutomationHealth: React.FC<{ handler: TestDataHandler}> = obser
                     .then(() => usersInfoCache.fetchUsersInfo(items.map(i => i.owner).filter(i => !!i) as string[]))
                     .then((fetchedUsers) => fetchedUsers.length > 0 && setHealthListState({items: items.slice(), columns: columns}));
 
-    }, [handler.filterUpdated]);
+    }, [handler.searchUpdated]);
 
     // find health item if not loaded already
     const testId = handler.getSearchParam('health') as string | undefined;
