[Horde][Analytics][Server] Cleaning up telemetry records to use base class for common attributes & supporting inheritance for AnalyticsTableGenerator.

#Synced-CL 49562268

diff --git a/Engine/Source/Programs/Horde/Plugins/Build/HordeServer.Build/Jobs/JobTelemetry.cs b/Engine/Source/Programs/Horde/Plugins/Build/HordeServer.Build/Jobs/JobTelemetry.cs
--- a/Engine/Source/Programs/Horde/Plugins/Build/HordeServer.Build/Jobs/JobTelemetry.cs
+++ b/Engine/Source/Programs/Horde/Plugins/Build/HordeServer.Build/Jobs/JobTelemetry.cs
@@ -1,6 +1,7 @@
 // Copyright Epic Games, Inc. All Rights Reserved.
 
 using System.ComponentModel.DataAnnotations.Schema;
+using EpicGames.Analytics;
 using EpicGames.Analytics.Telemetry;
 using EpicGames.Horde.Agents;
 using EpicGames.Horde.Agents.Pools;
@@ -18,20 +19,8 @@ namespace HordeServer.Jobs
 	/// </summary>
 	[AnalyticsTableGen]
 	[Table("ingest.horde.state_job_summary")]
-	public record JobSummaryTelemetry
+	public record JobSummaryTelemetry : AbstractTelemetryRecord
 	{
-		/// <summary>
-		/// The telemetry event name.
-		/// </summary>
-		[Column("event_name")]
-		public string EventName { get; init; }
-
-		/// <summary>
-		/// The schema version of the telemetry.
-		/// </summary>
-		[Column("schema_version")]
-		public int SchemaVersion { get; init; }
-
 		/// <summary>
 		/// The job Id.
 		/// </summary>
@@ -155,10 +144,8 @@ namespace HordeServer.Jobs
 		/// </summary>
 		/// <param name="job">The job to base this telemetry event on.</param>
 		/// <param name="finishTimeUtc">The completion time of this job.</param>
-		public JobSummaryTelemetry(IJob job, DateTime finishTimeUtc)
+		public JobSummaryTelemetry(IJob job, DateTime finishTimeUtc) : base(DefaultEventName, CurrentSchemaVersion)
 		{
-			EventName = DefaultEventName;
-			SchemaVersion = CurrentSchemaVersion;
 			JobId = job.Id;
 			TemplateId = job.TemplateId;
 			StreamId = job.StreamId;
@@ -186,10 +173,8 @@ namespace HordeServer.Jobs
 		/// Default constructor for an empty telemetry object,
 		/// </summary>
 		/// <remarks>This constructor is effectively reserved for ORM mapping.</remarks>
-		public JobSummaryTelemetry()
+		public JobSummaryTelemetry() : base(DefaultEventName, CurrentSchemaVersion)
 		{
-			EventName = DefaultEventName;
-			SchemaVersion = CurrentSchemaVersion;
 			CodeCommitId = String.Empty;
 			JobId = JobId.Empty;
 			TemplateId = new TemplateId(String.Empty);
diff --git a/Engine/Source/Programs/Shared/EpicGames.Analytics.Generators/AnalyticsTableGenerator.cs b/Engine/Source/Programs/Shared/EpicGames.Analytics.Generators/AnalyticsTableGenerator.cs
--- a/Engine/Source/Programs/Shared/EpicGames.Analytics.Generators/AnalyticsTableGenerator.cs
+++ b/Engine/Source/Programs/Shared/EpicGames.Analytics.Generators/AnalyticsTableGenerator.cs
@@ -146,26 +146,32 @@ namespace EpicGames.Analytics.Generators
 				}
 			}
 
-			foreach (IPropertySymbol prop in telemetrySymbol.GetMembers().OfType<IPropertySymbol>())
+			ITypeSymbol? current = telemetrySymbol;
+			while (current != null)
 			{
-				string propertyName = prop.Name;
-				string? columnName = null;
-
-				foreach (AttributeData attr in prop.GetAttributes())
+				foreach (IPropertySymbol prop in current.GetMembers().OfType<IPropertySymbol>())
 				{
-					if (attr.AttributeClass?.ToDisplayString() == "System.ComponentModel.DataAnnotations.Schema.ColumnAttribute")
+					string propertyName = prop.Name;
+					string? columnName = null;
+
+					foreach (AttributeData attr in prop.GetAttributes())
 					{
-						if (attr.ConstructorArguments.Length > 0)
+						if (attr.AttributeClass?.ToDisplayString() == "System.ComponentModel.DataAnnotations.Schema.ColumnAttribute")
 						{
-							columnName = attr.ConstructorArguments[0].Value as string;
+							if (attr.ConstructorArguments.Length > 0)
+							{
+								columnName = attr.ConstructorArguments[0].Value as string;
+							}
 						}
 					}
-				}
 
-				if (columnName != null)
-				{
-					columns.Add(new Tuple<string, string>(propertyName, columnName));
+					if (columnName != null)
+					{
+						columns.Add(new Tuple<string, string>(propertyName, columnName));
+					}
 				}
+
+				current = current.BaseType;
 			}
 
 			return String.IsNullOrEmpty(tableName) ? null : new AnalyticsTableToGenerate(telemetryName, tableName!, columns);
diff --git a/Engine/Source/Programs/Shared/EpicGames.Analytics.Telemetry/AbstractTelemetryRecord.cs b/Engine/Source/Programs/Shared/EpicGames.Analytics.Telemetry/AbstractTelemetryRecord.cs
new file mode 100644
--- /dev/null
+++ b/Engine/Source/Programs/Shared/EpicGames.Analytics.Telemetry/AbstractTelemetryRecord.cs
@@ -0,0 +1,35 @@
+// Copyright Epic Games, Inc. All Rights Reserved.
+
+using System.ComponentModel.DataAnnotations.Schema;
+
+namespace EpicGames.Analytics
+{
+	/// <summary>
+	/// Describes the base elements required for telemetry.
+	/// </summary>
+	public abstract record AbstractTelemetryRecord
+	{
+		/// <summary>
+		/// The telemetry event name.
+		/// </summary>
+		[Column("event_name")]
+		public string EventName { get; init; }
+
+		/// <summary>
+		/// The schema version of the telemetry.
+		/// </summary>
+		[Column("schema_version")]
+		public int SchemaVersion { get; init; }
+
+		/// <summary>
+		/// Constructor.
+		/// </summary>
+		/// <param name="eventName">The event name.</param>
+		/// <param name="schemaVersion">The schema version.</param>
+		protected AbstractTelemetryRecord(string eventName, int schemaVersion)
+		{
+			EventName = eventName;
+			SchemaVersion = schemaVersion;
+		}
+	}
+}
\ No newline at end of file
