[Horde][BuildHealth][Server] Create jobid,batchid,stepid column hierarchy within JobStepRef.

#Synced-CL 49602718

diff --git a/Engine/Source/Programs/Horde/Plugins/Build/HordeServer.Build/Jobs/JobStepRefCollection.cs b/Engine/Source/Programs/Horde/Plugins/Build/HordeServer.Build/Jobs/JobStepRefCollection.cs
--- a/Engine/Source/Programs/Horde/Plugins/Build/HordeServer.Build/Jobs/JobStepRefCollection.cs
+++ b/Engine/Source/Programs/Horde/Plugins/Build/HordeServer.Build/Jobs/JobStepRefCollection.cs
@@ -29,6 +29,12 @@ namespace HordeServer.Jobs
 			[BsonId]
 			public JobStepRefId Id { get; set; }
 
+			public JobId? JobId { get; set; }
+
+			public JobStepBatchId? BatchId { get; set; }
+
+			public JobStepId? StepId { get; set; }
+
 			public string JobName { get; set; } = "Unknown";
 
 			public string Name { get; set; }
@@ -114,9 +120,12 @@ namespace HordeServer.Jobs
 			IReadOnlyList<int>? IJobStepRef.IssueIds => IssueIds;
 			IReadOnlyList<string>? IJobStepRef.Metadata => Metadata;
 
-			public JobStepRef(JobStepRefId id, string jobName, string nodeName, StreamId streamId, TemplateId templateId, CommitIdWithOrder commitId, LogId? logId, PoolId? poolId, AgentId? agentId, JobStepState? state, JobStepOutcome? outcome, bool updateIssues, CommitIdWithOrder? lastSuccess, CommitIdWithOrder? lastWarning, float batchWaitTime, float batchInitTime, DateTime jobStartTimeUtc, DateTime startTimeUtc, DateTime? finishTimeUtc, List<string>? stepMetadata = null)
+			public JobStepRef(JobStepRefId id, JobId? jobId, JobStepBatchId? batchId, JobStepId? stepId, string jobName, string nodeName, StreamId streamId, TemplateId templateId, CommitIdWithOrder commitId, LogId? logId, PoolId? poolId, AgentId? agentId, JobStepState? state, JobStepOutcome? outcome, bool updateIssues, CommitIdWithOrder? lastSuccess, CommitIdWithOrder? lastWarning, float batchWaitTime, float batchInitTime, DateTime jobStartTimeUtc, DateTime startTimeUtc, DateTime? finishTimeUtc, List<string>? stepMetadata = null)
 			{
 				Id = id;
+				JobId = jobId;
+				BatchId = batchId;
+				StepId = stepId;
 				JobName = jobName;
 				Name = nodeName;
 				StreamId = streamId;
@@ -151,6 +160,13 @@ namespace HordeServer.Jobs
 			List<MongoIndex<JobStepRef>> indexes = new List<MongoIndex<JobStepRef>>();
 			indexes.Add(keys => keys.Ascending(x => x.StreamId).Ascending(x => x.TemplateId).Ascending(x => x.Name).Descending(x => x.CommitOrder));
 
+			// Index for obId, batchId, stepId hierarchical queries.
+			// Using this index, we can cost-effectively query:
+			//	- All job step refs belonging to a specific job
+			//	- All job step refs belonging to a batch (of a job)
+			// It's sparse in nature due for backwards compatibility, and it is accepted that some items which use this index will miss items.
+			indexes.Add(keys => keys.Ascending(x => x.JobId).Ascending(x => x.BatchId).Ascending(x => x.StepId), sparse: true);
+
 			_jobStepRefs = mongoService.GetCollection<JobStepRef>("JobStepRefs", indexes);
 			_telemetryWriter = telemetryWriter;
 			_buildConfig = buildConfig;
@@ -159,7 +175,7 @@ namespace HordeServer.Jobs
 		/// <inheritdoc/>
 		public async Task<IJobStepRef> InsertOrReplaceAsync(JobStepRefId id, string jobName, string stepName, StreamId streamId, TemplateId templateId, CommitIdWithOrder commitId, LogId? logId, PoolId? poolId, AgentId? agentId, JobStepState? state, JobStepOutcome? outcome, bool updateIssues, CommitIdWithOrder? lastSuccess, CommitIdWithOrder? lastWarning, float waitTime, float initTime, DateTime jobStartTimeUtc, DateTime startTimeUtc, DateTime? finishTimeUtc, List<string>? stepMetadata = null, ILogger? logger = null)
 		{
-			JobStepRef newJobStepRef = new JobStepRef(id, jobName, stepName, streamId, templateId, commitId, logId, poolId, agentId, state, outcome, updateIssues, lastSuccess, lastWarning, waitTime, initTime, jobStartTimeUtc, startTimeUtc, finishTimeUtc, stepMetadata);
+			JobStepRef newJobStepRef = new JobStepRef(id, id.JobId, id.BatchId, id.StepId, jobName, stepName, streamId, templateId, commitId, logId, poolId, agentId, state, outcome, updateIssues, lastSuccess, lastWarning, waitTime, initTime, jobStartTimeUtc, startTimeUtc, finishTimeUtc, stepMetadata);
 			await _jobStepRefs.ReplaceOneAsync(Builders<JobStepRef>.Filter.Eq(x => x.Id, newJobStepRef.Id), newJobStepRef, new ReplaceOptions { IsUpsert = true });
 
 			try
diff --git a/Engine/Source/Programs/Shared/EpicGames.Horde/Jobs/IJobStepRef.cs b/Engine/Source/Programs/Shared/EpicGames.Horde/Jobs/IJobStepRef.cs
--- a/Engine/Source/Programs/Shared/EpicGames.Horde/Jobs/IJobStepRef.cs
+++ b/Engine/Source/Programs/Shared/EpicGames.Horde/Jobs/IJobStepRef.cs
@@ -22,6 +22,24 @@ namespace HordeServer.Jobs
 		/// </summary>
 		public JobStepRefId Id { get; }
 
+		/// <summary>
+		/// The job the jobstep belongs to.
+		/// </summary>
+		/// <remarks>This is optional due to backwards compatibility for previously serialized items.</remarks>
+		public JobId? JobId { get; }
+
+		/// <summary>
+		/// The batch the jobstep belongs to.
+		/// </summary>
+		/// <remarks>This is optional due to backwards compatibility for previously serialized items.</remarks>
+		public JobStepBatchId? BatchId { get; }
+
+		/// <summary>
+		/// The step the job step represents.
+		/// </summary>
+		/// <remarks>This is optional due to backwards compatibility for previously serialized items.</remarks>
+		public JobStepId? StepId { get; }
+
 		/// <summary>
 		/// Name of the job
 		/// </summary>
