Horde: The task argument for jobs and conforms is passed using an environment variable

#Synced-CL 0

diff --git a/Engine/Source/Programs/Horde/Drivers/JobDriver/DriverApp.cs b/Engine/Source/Programs/Horde/Drivers/JobDriver/DriverApp.cs
--- a/Engine/Source/Programs/Horde/Drivers/JobDriver/DriverApp.cs
+++ b/Engine/Source/Programs/Horde/Drivers/JobDriver/DriverApp.cs
@@ -22,6 +22,14 @@ namespace JobDriver
 		/// </summary>
 		public static async Task<int> Main(string[] args)
 		{
+			string? taskArg = Environment.GetEnvironmentVariable("UE_HORDE_JOB_DRIVER_TASK_ARG");
+			if (!String.IsNullOrEmpty(taskArg))
+			{
+				Array.Resize(ref args, args.Length + 1);
+				args[^1] = taskArg;
+				Environment.SetEnvironmentVariable("UE_HORDE_JOB_DRIVER_TASK_ARG", null);
+			}
+
 			CommandLineArguments arguments = new CommandLineArguments(args);
 
 			// Create the services
diff --git a/Engine/Source/Programs/Horde/HordeAgent/Leases/Handlers/ConformHandler.cs b/Engine/Source/Programs/Horde/HordeAgent/Leases/Handlers/ConformHandler.cs
--- a/Engine/Source/Programs/Horde/HordeAgent/Leases/Handlers/ConformHandler.cs
+++ b/Engine/Source/Programs/Horde/HordeAgent/Leases/Handlers/ConformHandler.cs
@@ -34,11 +34,11 @@ namespace HordeAgent.Leases.Handlers
 			arguments.Add($"-SessionId={session.SessionId}");
 			arguments.Add($"-LeaseId={leaseId}");
 			arguments.Add($"-WorkingDir={session.WorkingDir}");
-			arguments.Add($"-Task={Convert.ToBase64String(conformTask.ToByteArray())}");
 
 			FileReference driverAssembly = FileReference.Combine(new DirectoryReference(AppContext.BaseDirectory), "JobDriver", "JobDriver.dll");
 
 			Dictionary<string, string> environment = ManagedProcess.GetCurrentEnvVars();
+			environment["UE_HORDE_JOB_DRIVER_TASK_ARG"] = $"-Task={Convert.ToBase64String(conformTask.ToByteArray())}";
 			environment[HordeHttpClient.HordeUrlEnvVarName] = session.HordeClient.ServerUrl.ToString();
 			environment[HordeHttpClient.HordeTokenEnvVarName] = await session.HordeClient.GetAccessTokenAsync(false, cancellationToken) ?? String.Empty;
 			environment["UE_LOG_JSON_TO_STDOUT"] = "1";
diff --git a/Engine/Source/Programs/Horde/HordeAgent/Leases/Handlers/JobHandler.cs b/Engine/Source/Programs/Horde/HordeAgent/Leases/Handlers/JobHandler.cs
--- a/Engine/Source/Programs/Horde/HordeAgent/Leases/Handlers/JobHandler.cs
+++ b/Engine/Source/Programs/Horde/HordeAgent/Leases/Handlers/JobHandler.cs
@@ -113,13 +113,13 @@ namespace HordeAgent.Leases.Handlers
 			arguments.Add($"-SessionId={session.SessionId}");
 			arguments.Add($"-LeaseId={leaseId}");
 			arguments.Add($"-WorkingDir={session.WorkingDir}");
-			arguments.Add($"-Task={Convert.ToBase64String(executeTask.ToByteArray())}");
 
 			string driverName = String.IsNullOrEmpty(executeTask.JobOptions.Driver) ? "JobDriver" : executeTask.JobOptions.Driver;			
 			FileReference driverAssembly = FileReference.Combine(new DirectoryReference(AppContext.BaseDirectory), driverName, $"{driverName}.dll");			
 			span.SetAttribute("horde.job.driver_name", driverName);
 
 			Dictionary<string, string> environment = ManagedProcess.GetCurrentEnvVars();
+			environment["UE_HORDE_JOB_DRIVER_TASK_ARG"] = $"-Task={Convert.ToBase64String(executeTask.ToByteArray())}";
 
 			try
 			{
