Horde: Conforms that are scheduled due to low free disk space upon execution check if the agent's available disk space still requires the conform to be run

#Synced-CL 0

diff --git a/Engine/Source/Programs/Horde/Drivers/JobDriver.Tests/Execution/ConformExecutorTests.cs b/Engine/Source/Programs/Horde/Drivers/JobDriver.Tests/Execution/ConformExecutorTests.cs
--- a/Engine/Source/Programs/Horde/Drivers/JobDriver.Tests/Execution/ConformExecutorTests.cs
+++ b/Engine/Source/Programs/Horde/Drivers/JobDriver.Tests/Execution/ConformExecutorTests.cs
@@ -2,8 +2,10 @@
 
 using System;
 using System.Collections.Generic;
+using System.IO;
 using System.Threading;
 using System.Threading.Tasks;
+using EpicGames.Core;
 using EpicGames.Horde.Agents;
 using EpicGames.Horde.Agents.Leases;
 using EpicGames.Perforce.Fixture;
@@ -12,6 +14,7 @@ using HordeCommon.Rpc.Tasks;
 using JobDriver.Execution;
 using Microsoft.Extensions.DependencyInjection;
 using Microsoft.Extensions.Logging;
+using Microsoft.Extensions.Logging.Abstractions;
 using Microsoft.VisualStudio.TestTools.UnitTesting;
 using OpenTelemetry.Trace;
 using SyncOptions = JobDriver.Execution.SyncOptions;
@@ -99,7 +102,46 @@ public class ConformTest : BasePerforceFixtureTest
 		FakeHordeClient hordeClient = new (jobRpcClientStub);
 		return new ConformExecutor(hordeClient, _materializerFactories, TempDir, new AgentId("agent1"), leaseId, conformTask,
 			new DriverSettings(),
+			new FakeSystemMetrics(),
 			NoOpTracer, _logger);
 	}
 }
 
+[TestClass]
+public class ConformNoPerforceTest
+{
+	private static Tracer NoOpTracer { get; } = TracerProvider.Default.GetTracer("NoOp");
+
+	[TestMethod]
+	[DataRow(500, true, DisplayName = "Conform skipped because there is enough disk free space available")]
+	[DataRow(5000, false, DisplayName = "Conform not skipped because there is not enough disk free space available")]
+	[DataRow(0, false, DisplayName = "Conform not skipped because ConformDiskFreeSpace workspace property is not set")]
+	public async Task SkipConformWhenAvailableFreeDiskSpaceAsync(long conformDiskFreeSpace, bool skippedResult)
+	{
+		ConformTask conformTask = new()
+		{
+			SkipWhenDiskFreeSpaceIsAvailable = true,
+			Workspaces = { new List<RpcAgentWorkspace> { new() { ConformDiskFreeSpace = conformDiskFreeSpace } } }
+		};
+
+		// Required because not using BasePerforceFixtureTest and must not run p4 commands;
+		// to avoid it using the local machine's p4 connection settings
+		DriverSettings driverSettings = new() { Executor = String.Empty, PerforceExecutor = { RunConform = false } };
+
+		ConformExecutor executor = new(
+			new FakeHordeClient(new JobRpcClientStub(NullLogger.Instance)),
+			[],
+			new DirectoryReference(Path.GetTempPath()),
+			new AgentId("agent1"),
+			LeaseId.Parse("111111111111111111111111"),
+			conformTask,
+			driverSettings,
+			new FakeSystemMetrics { FreeSpace = 1000 },
+			NoOpTracer,
+			NullLogger.Instance);
+
+		await executor.ExecuteAsync(CancellationToken.None);
+
+		Assert.AreEqual(skippedResult, executor.ConformWasSkipped);
+	}
+}
diff --git a/Engine/Source/Programs/Horde/Drivers/JobDriver.Tests/Execution/FakeSystemMetrics.cs b/Engine/Source/Programs/Horde/Drivers/JobDriver.Tests/Execution/FakeSystemMetrics.cs
new file mode 100644
--- /dev/null
+++ b/Engine/Source/Programs/Horde/Drivers/JobDriver.Tests/Execution/FakeSystemMetrics.cs
@@ -0,0 +1,26 @@
+// Copyright Epic Games, Inc. All Rights Reserved.
+
+using EpicGames.Horde.Agents.Metrics;
+
+namespace JobDriver.Tests.Execution;
+
+internal class FakeSystemMetrics : ISystemMetrics
+{
+	public long FreeSpace { get; init; }
+	public long TotalSize { get; init; }
+
+	/// <inheritdoc/>
+	public CpuMetrics? GetCpu() => null;
+
+	/// <inheritdoc/>
+	public CpuTimeMetrics? GetCpuTime() => null;
+
+	/// <inheritdoc/>
+	public MemoryMetrics? GetMemory() => null;
+
+	/// <inheritdoc/>
+	public DiskMetrics? GetDisk()
+	{
+		return new DiskMetrics { FreeSpace = FreeSpace, TotalSize = TotalSize };
+	}
+}
diff --git a/Engine/Source/Programs/Horde/Drivers/JobDriver/Commands/Execution/ExecuteConformCommand.cs b/Engine/Source/Programs/Horde/Drivers/JobDriver/Commands/Execution/ExecuteConformCommand.cs
--- a/Engine/Source/Programs/Horde/Drivers/JobDriver/Commands/Execution/ExecuteConformCommand.cs
+++ b/Engine/Source/Programs/Horde/Drivers/JobDriver/Commands/Execution/ExecuteConformCommand.cs
@@ -4,6 +4,7 @@ using EpicGames.Core;
 using EpicGames.Horde;
 using EpicGames.Horde.Agents;
 using EpicGames.Horde.Agents.Leases;
+using EpicGames.Horde.Agents.Metrics;
 using EpicGames.Horde.Agents.Sessions;
 using JobDriver.Execution;
 using HordeCommon.Rpc.Tasks;
@@ -34,13 +35,15 @@ namespace JobDriver.Commands.Execution
 		readonly HordeClientFactory _hordeClientFactory;
 		readonly IEnumerable<IWorkspaceMaterializerFactory> _materializerFactories;
 		readonly IOptions<DriverSettings> _driverSettings;
+		readonly ISystemMetrics _systemMetrics;
 		readonly Tracer _tracer;
 
-		public ExecuteConformCommand(HordeClientFactory hordeClientFactory, IEnumerable<IWorkspaceMaterializerFactory> materializerFactories, IOptions<DriverSettings> driverSettings, Tracer tracer)
+		public ExecuteConformCommand(HordeClientFactory hordeClientFactory, IEnumerable<IWorkspaceMaterializerFactory> materializerFactories, IOptions<DriverSettings> driverSettings, ISystemMetrics systemMetrics, Tracer tracer)
 		{
 			_hordeClientFactory = hordeClientFactory;
 			_materializerFactories = materializerFactories;
 			_driverSettings = driverSettings;
+			_systemMetrics = systemMetrics;
 			_tracer = tracer;
 		}
 
@@ -52,7 +55,7 @@ namespace JobDriver.Commands.Execution
 
 			await using HordeClient hordeClient = _hordeClientFactory.Create();
 
-			ConformExecutor conformExecutor = new ConformExecutor(hordeClient, _materializerFactories, WorkingDir, AgentId, LeaseId, conformTask, _driverSettings.Value, _tracer, logger);
+			ConformExecutor conformExecutor = new ConformExecutor(hordeClient, _materializerFactories, WorkingDir, AgentId, LeaseId, conformTask, _driverSettings.Value, _systemMetrics, _tracer, logger);
 			await conformExecutor.ExecuteAsync(CancellationToken.None);
 
 			return 0;
diff --git a/Engine/Source/Programs/Horde/Drivers/JobDriver/DriverApp.cs b/Engine/Source/Programs/Horde/Drivers/JobDriver/DriverApp.cs
--- a/Engine/Source/Programs/Horde/Drivers/JobDriver/DriverApp.cs
+++ b/Engine/Source/Programs/Horde/Drivers/JobDriver/DriverApp.cs
@@ -3,11 +3,13 @@
 using System.Reflection;
 using EpicGames.Core;
 using EpicGames.Horde;
+using EpicGames.Horde.Agents.Metrics;
 using EpicGames.Horde.Utilities;
 using JobDriver.Execution;
 using JobDriver.Utility;
 using Microsoft.Extensions.Configuration;
 using Microsoft.Extensions.DependencyInjection;
+using Microsoft.Extensions.Logging;
 using OpenTelemetry.Trace;
 
 namespace JobDriver
@@ -34,7 +36,7 @@ namespace JobDriver
 
 			// Create the services
 			IServiceCollection services = new ServiceCollection();
-			RegisterServices(services);
+			RegisterServices(services, arguments);
 
 			// Run the host
 			await using ServiceProvider serviceProvider = services.BuildServiceProvider();
@@ -52,6 +54,16 @@ namespace JobDriver
 		/// </summary>
 		/// <param name="services"></param>
 		public static void RegisterServices(IServiceCollection services)
+		{
+			RegisterServices(services, new CommandLineArguments([]));
+		}
+
+		/// <summary>
+		/// Helper method to register services for this app
+		/// </summary>
+		/// <param name="services"></param>
+		/// <param name="arguments"></param>
+		public static void RegisterServices(IServiceCollection services, CommandLineArguments arguments)
 		{
 			// Read the driver config
 			IConfiguration configuration = new ConfigurationBuilder()
@@ -91,7 +103,20 @@ namespace JobDriver
 			services.AddSingleton<IWorkspaceMaterializerFactory, ManagedWorkspaceMaterializerFactory>();
 			services.AddSingleton<IWorkspaceMaterializerFactory, PerforceMaterializerFactory>();
 
+			arguments.TryGetValue("-WorkingDir=", out string? workingDir);
+			services.AddSingleton<ISystemMetrics>(sp => CreateSystemMetrics(workingDir, sp.GetRequiredService<ILogger<ISystemMetrics>>()));
+
 			services.AddCommandsFromAssembly(Assembly.GetExecutingAssembly());
 		}
+
+		static ISystemMetrics CreateSystemMetrics(string? workingDir, ILogger logger)
+		{
+			DirectoryReference? wd = null;
+			if (workingDir != null)
+			{
+				wd = new DirectoryReference(workingDir);
+			}
+			return new DefaultSystemMetrics(wd, logger);
+		}
 	}
 }
diff --git a/Engine/Source/Programs/Horde/Drivers/JobDriver/Execution/ConformExecutor.cs b/Engine/Source/Programs/Horde/Drivers/JobDriver/Execution/ConformExecutor.cs
--- a/Engine/Source/Programs/Horde/Drivers/JobDriver/Execution/ConformExecutor.cs
+++ b/Engine/Source/Programs/Horde/Drivers/JobDriver/Execution/ConformExecutor.cs
@@ -4,6 +4,7 @@ using EpicGames.Core;
 using EpicGames.Horde;
 using EpicGames.Horde.Agents;
 using EpicGames.Horde.Agents.Leases;
+using EpicGames.Horde.Agents.Metrics;
 using JobDriver.Utility;
 using Horde.Common.Rpc;
 using HordeCommon.Rpc.Messages;
@@ -22,11 +23,20 @@ namespace JobDriver.Execution
 		readonly IEnumerable<IWorkspaceMaterializerFactory> _materializerFactories;
 		readonly ConformTask _conformTask;
 		readonly DriverSettings _driverSettings;
+		readonly ISystemMetrics _systemMetrics;
 		readonly Tracer _tracer;
 		readonly ILogger _logger;
 
+		/// <summary>
+		/// Whether it was decided to skip the conform during execution.
+		/// </summary>
+		/// <remarks>
+		/// Exposed to ease testing.
+		/// </remarks>
+		internal bool ConformWasSkipped { get; private set; }
+
 		public ConformExecutor(IHordeClient hordeClient, IEnumerable<IWorkspaceMaterializerFactory> materializerFactories,
-			DirectoryReference workingDir, AgentId agentId, LeaseId leaseId, ConformTask conformTask, DriverSettings driverSettings, Tracer tracer, ILogger logger)
+			DirectoryReference workingDir, AgentId agentId, LeaseId leaseId, ConformTask conformTask, DriverSettings driverSettings, ISystemMetrics systemMetrics, Tracer tracer, ILogger logger)
 		{
 			_hordeClient = hordeClient;
 			_workingDir = workingDir;
@@ -35,6 +45,7 @@ namespace JobDriver.Execution
 			_materializerFactories = materializerFactories;
 			_conformTask = conformTask;
 			_driverSettings = driverSettings;
+			_systemMetrics = systemMetrics;
 			_tracer = tracer;
 			_logger = logger;
 		}
@@ -103,7 +114,12 @@ namespace JobDriver.Execution
 					.Select(x => DirectoryReference.Combine(_workingDir, x.Identifier)).ToList();
 
 				// Run the conform task
-				if (isExecutorConformCompatible && _driverSettings.PerforceExecutor.RunConform)
+				if (_conformTask.SkipWhenDiskFreeSpaceIsAvailable && HasAvailableDiskFreeSpace(pendingWorkspaces))
+				{
+					ConformWasSkipped = true;
+					_logger.LogInformation("Skipping conform. Executor={Executor} RunConform={RunConform}", _driverSettings.Executor, _driverSettings.PerforceExecutor.RunConform);
+				}
+				else if (isExecutorConformCompatible && _driverSettings.PerforceExecutor.RunConform)
 				{
 					if (managedWorkspaces.Count > 0)
 					{
@@ -141,5 +157,33 @@ namespace JobDriver.Execution
 				removeUntrackedFiles = response.RemoveUntrackedFiles;
 			}
 		}
+
+		bool HasAvailableDiskFreeSpace(IList<RpcAgentWorkspace> pendingWorkspaces)
+		{
+			long? freeDiskSpace = _systemMetrics.GetDisk()?.FreeSpace;
+			if (freeDiskSpace == null)
+			{
+				return false;
+			}
+
+			long conformDiskSpaceNeeded = 0;
+			foreach (RpcAgentWorkspace workspace in pendingWorkspaces)
+			{
+				if (workspace.ConformDiskFreeSpace > 0)
+				{
+					conformDiskSpaceNeeded = Math.Max(conformDiskSpaceNeeded, workspace.ConformDiskFreeSpace);
+				}
+			}
+
+			bool hasAvailableFreeDiskSpace = conformDiskSpaceNeeded > 0 && freeDiskSpace >= conformDiskSpaceNeeded;
+			if (hasAvailableFreeDiskSpace)
+			{
+				double BytesToMegabytes(long v) => v / 1024.0 / 1024.0;
+				_logger.LogInformation("Skipping conform as free disk space available ({FreeDiskSpace:F1} MB) is greater than the required free disk space when the conform was requested ({RequiredFreeDiskSpace:F1} MB)",
+					BytesToMegabytes(freeDiskSpace.Value), BytesToMegabytes(conformDiskSpaceNeeded));
+			}
+
+			return hasAvailableFreeDiskSpace;
+		}
 	}
 }
diff --git a/Engine/Source/Programs/Horde/Drivers/JobDriver/Execution/DefaultSystemMetrics.cs b/Engine/Source/Programs/Horde/Drivers/JobDriver/Execution/DefaultSystemMetrics.cs
new file mode 100644
--- /dev/null
+++ b/Engine/Source/Programs/Horde/Drivers/JobDriver/Execution/DefaultSystemMetrics.cs
@@ -0,0 +1,74 @@
+// Copyright Epic Games, Inc. All Rights Reserved.
+
+using System.Runtime.InteropServices;
+using EpicGames.Core;
+using EpicGames.Horde.Agents.Metrics;
+using Microsoft.Extensions.Logging;
+
+namespace JobDriver.Execution;
+
+/// <summary>
+/// Implementation of <see cref="ISystemMetrics"/>
+/// </summary>
+public class DefaultSystemMetrics : ISystemMetrics
+{
+	private readonly DirectoryReference? _workingDir;
+	private readonly ILogger _logger;
+
+	/// <summary>
+	/// Constructor
+	/// </summary>
+	/// <param name="workingDir"></param>
+	/// <param name="logger"></param>
+	public DefaultSystemMetrics(DirectoryReference? workingDir, ILogger logger)
+	{
+		_workingDir = workingDir;
+		_logger = logger;
+	}
+
+	/// <inheritdoc/>
+	public CpuMetrics? GetCpu() => null;
+
+	/// <inheritdoc/>
+	public CpuTimeMetrics? GetCpuTime() => null;
+
+	/// <inheritdoc/>
+	public MemoryMetrics? GetMemory() => null;
+
+	/// <inheritdoc/>
+	public DiskMetrics? GetDisk()
+	{
+		if (_workingDir == null)
+		{
+			_logger.LogError("Cannot get disk metrics: no working directory");
+			return null;
+		}
+
+		try
+		{
+			string? rootDir;
+			if (RuntimeInformation.IsOSPlatform(OSPlatform.Windows))
+			{
+				rootDir = Path.GetPathRoot(_workingDir.FullName);
+			}
+			else
+			{
+				rootDir = _workingDir.FullName;
+			}
+
+			if (rootDir == null)
+			{
+				_logger.LogError("Cannot get disk metrics: no root directory");
+				return null;
+			}
+
+			DriveInfo driveInfo = new(rootDir);
+			return new DiskMetrics { FreeSpace = driveInfo.AvailableFreeSpace, TotalSize = driveInfo.TotalSize };
+		}
+		catch (Exception error)
+		{
+			_logger.LogError(error, "Cannot get disk metrics: {Message}", error.Message);
+			return null;
+		}
+	}
+}
diff --git a/Engine/Source/Programs/Horde/Plugins/Build/HordeServer.Build.Tests/Agents/AgentControllerTest.cs b/Engine/Source/Programs/Horde/Plugins/Build/HordeServer.Build.Tests/Agents/AgentControllerTest.cs
--- a/Engine/Source/Programs/Horde/Plugins/Build/HordeServer.Build.Tests/Agents/AgentControllerTest.cs
+++ b/Engine/Source/Programs/Horde/Plugins/Build/HordeServer.Build.Tests/Agents/AgentControllerTest.cs
@@ -27,4 +27,118 @@ public class AgentControllerDbTest : BuildTestSetup
 		getRes = (obj.Value as GetAgentResponse)!;
 		Assert.AreEqual("foo bar baz", getRes.Comment);
 	}
+
+	[TestMethod]
+	public async Task UpdateAgentFullConformWithDiskSpaceNeededAsync()
+	{
+		Fixture fixture = await CreateFixtureAsync();
+		IAgent agent = fixture.Agent1;
+
+		GetAgentResponse response = Deref((await AgentsController.GetAgentAsync(agent.Id)).Value as GetAgentResponse);
+		Assert.IsFalse(response.PendingFullConform);
+		Assert.IsFalse(response.PendingConformSkipWhenDiskFreeSpaceIsAvailable);
+
+		await AgentsController.UpdateAgentAsync(agent.Id, new UpdateAgentRequest(RequestFullConform: true, RequestConformSkipWhenAvailableFreeDiskSpace: true));
+
+		response = Deref((await AgentsController.GetAgentAsync(agent.Id)).Value as GetAgentResponse);
+		Assert.IsTrue(response.PendingFullConform);
+		Assert.IsTrue(response.PendingConformSkipWhenDiskFreeSpaceIsAvailable);
+	}
+
+	[TestMethod]
+	public async Task UpdateAgentFullConformWithoutDiskSpaceNeededAsync()
+	{
+		Fixture fixture = await CreateFixtureAsync();
+		IAgent agent = fixture.Agent1;
+
+		GetAgentResponse response = Deref((await AgentsController.GetAgentAsync(agent.Id)).Value as GetAgentResponse);
+		Assert.IsFalse(response.PendingFullConform);
+		Assert.IsFalse(response.PendingConformSkipWhenDiskFreeSpaceIsAvailable);
+
+		await AgentsController.UpdateAgentAsync(agent.Id, new UpdateAgentRequest(RequestFullConform: true));
+
+		response = Deref((await AgentsController.GetAgentAsync(agent.Id)).Value as GetAgentResponse);
+		Assert.IsTrue(response.PendingFullConform);
+		Assert.IsFalse(response.PendingConformSkipWhenDiskFreeSpaceIsAvailable);
+	}
+
+	[TestMethod]
+	public async Task UpdateAgentIgnoreDiskSpaceNeededWhenNotFullConformAsync()
+	{
+		Fixture fixture = await CreateFixtureAsync();
+		IAgent agent = fixture.Agent1;
+
+		GetAgentResponse response = Deref((await AgentsController.GetAgentAsync(agent.Id)).Value as GetAgentResponse);
+		Assert.IsFalse(response.PendingFullConform);
+		Assert.IsFalse(response.PendingConform);
+		Assert.IsFalse(response.PendingConformSkipWhenDiskFreeSpaceIsAvailable);
+
+		// check disk space option is ignored when a conform is not requested
+		await AgentsController.UpdateAgentAsync(agent.Id, new UpdateAgentRequest(RequestConformSkipWhenAvailableFreeDiskSpace: true));
+
+		response = Deref((await AgentsController.GetAgentAsync(agent.Id)).Value as GetAgentResponse);
+		Assert.IsFalse(response.PendingFullConform);
+		Assert.IsFalse(response.PendingConform);
+		Assert.IsFalse(response.PendingConformSkipWhenDiskFreeSpaceIsAvailable);
+
+		// When not a full conform check disk space option should be ignored because untracked files won't be deleted
+		await AgentsController.UpdateAgentAsync(agent.Id, new UpdateAgentRequest(RequestConform: true, RequestConformSkipWhenAvailableFreeDiskSpace: true));
+
+		response = Deref((await AgentsController.GetAgentAsync(agent.Id)).Value as GetAgentResponse);
+		Assert.IsTrue(response.PendingConform);
+		Assert.IsFalse(response.PendingConformSkipWhenDiskFreeSpaceIsAvailable);
+	}
+
+	[TestMethod]
+	public async Task UpdateAgentSuccessiveFullConformRemovePreviousDiskSpaceNeededAsync()
+	{
+		Fixture fixture = await CreateFixtureAsync();
+		IAgent agent = fixture.Agent1;
+
+		GetAgentResponse response = Deref((await AgentsController.GetAgentAsync(agent.Id)).Value as GetAgentResponse);
+		Assert.IsFalse(response.PendingFullConform);
+		Assert.IsFalse(response.PendingConformSkipWhenDiskFreeSpaceIsAvailable);
+
+		await AgentsController.UpdateAgentAsync(agent.Id, new UpdateAgentRequest(RequestFullConform: true, RequestConformSkipWhenAvailableFreeDiskSpace: true));
+
+		response = Deref((await AgentsController.GetAgentAsync(agent.Id)).Value as GetAgentResponse);
+		Assert.IsTrue(response.PendingFullConform);
+		Assert.IsTrue(response.PendingConformSkipWhenDiskFreeSpaceIsAvailable);
+
+		// Successive full conform request without an explicit check disk space option will remove it
+		await AgentsController.UpdateAgentAsync(agent.Id, new UpdateAgentRequest(RequestFullConform: true));
+
+		response = Deref((await AgentsController.GetAgentAsync(agent.Id)).Value as GetAgentResponse);
+		Assert.IsNotNull(response);
+		Assert.IsTrue(response.PendingFullConform);
+		Assert.IsFalse(response.PendingConformSkipWhenDiskFreeSpaceIsAvailable);
+	}
+
+	[TestMethod]
+	public async Task UpdateAgentSuccessiveConformRemovePreviousDiskSpaceNeededAsync()
+	{
+		Fixture fixture = await CreateFixtureAsync();
+		IAgent agent = fixture.Agent1;
+
+		GetAgentResponse response = Deref((await AgentsController.GetAgentAsync(agent.Id)).Value as GetAgentResponse);
+		Assert.IsFalse(response.PendingFullConform);
+		Assert.IsFalse(response.PendingConform);
+		Assert.IsFalse(response.PendingConformSkipWhenDiskFreeSpaceIsAvailable);
+
+		await AgentsController.UpdateAgentAsync(agent.Id, new UpdateAgentRequest(RequestFullConform: true, RequestConformSkipWhenAvailableFreeDiskSpace: true));
+
+		response = Deref((await AgentsController.GetAgentAsync(agent.Id)).Value as GetAgentResponse);
+		Assert.IsTrue(response.PendingFullConform);
+		Assert.IsTrue(response.PendingConformSkipWhenDiskFreeSpaceIsAvailable);
+
+		// Successive conform request will remove the check disk space option
+		await AgentsController.UpdateAgentAsync(agent.Id, new UpdateAgentRequest(RequestConform: true));
+
+		response = Deref((await AgentsController.GetAgentAsync(agent.Id)).Value as GetAgentResponse);
+		Assert.IsNotNull(response);
+		Assert.IsTrue(response.PendingConform);
+		// The previous full conform remains but check disk space needed is removed by the new conform request
+		Assert.IsTrue(response.PendingFullConform);
+		Assert.IsFalse(response.PendingConformSkipWhenDiskFreeSpaceIsAvailable);
+	}
 }
\ No newline at end of file
diff --git a/Engine/Source/Programs/Horde/Plugins/Build/HordeServer.Build.Tests/Agents/AgentServiceTest.cs b/Engine/Source/Programs/Horde/Plugins/Build/HordeServer.Build.Tests/Agents/AgentServiceTest.cs
--- a/Engine/Source/Programs/Horde/Plugins/Build/HordeServer.Build.Tests/Agents/AgentServiceTest.cs
+++ b/Engine/Source/Programs/Horde/Plugins/Build/HordeServer.Build.Tests/Agents/AgentServiceTest.cs
@@ -245,4 +245,22 @@ public class AgentServiceTest : BuildTestSetup
 		return new ClaimsPrincipal(new ClaimsIdentity(
 			new List<Claim> { new(HordeClaimTypes.AgentSessionId, agent.SessionId.ToString()!) }, "TestAuthType"));
 	}
+
+	[TestMethod]
+	public async Task ConformRemovedAfterUpdateWorkspacesAsync()
+	{
+		Fixture fixture = await CreateFixtureAsync();
+		AgentId agentId = new(fixture.Agent1Name);
+
+		IAgent agent = Deref(await AgentService.GetAgentAsync(agentId));
+		agent = Deref(await agent.TryUpdateAsync(new UpdateAgentOptions(RequestFullConform: true, RequestConformSkipWhenDiskFreeSpaceIsAvailable: true)));
+		Assert.IsTrue(agent.RequestFullConform);
+		Assert.IsTrue(agent.RequestConformSkipWhenDiskFreeSpaceIsAvailable);
+
+		await AgentService.TryUpdateWorkspacesAsync(agent, agent.Workspaces.ToList(), false, CancellationToken.None);
+
+		agent = Deref(await AgentService.GetAgentAsync(agentId));
+		Assert.IsFalse(agent.RequestFullConform);
+		Assert.IsFalse(agent.RequestConformSkipWhenDiskFreeSpaceIsAvailable);
+	}
 }
\ No newline at end of file
diff --git a/Engine/Source/Programs/Horde/Plugins/Build/HordeServer.Build.Tests/Agents/Pools/PoolUpdateServiceTest.cs b/Engine/Source/Programs/Horde/Plugins/Build/HordeServer.Build.Tests/Agents/Pools/PoolUpdateServiceTest.cs
--- a/Engine/Source/Programs/Horde/Plugins/Build/HordeServer.Build.Tests/Agents/Pools/PoolUpdateServiceTest.cs
+++ b/Engine/Source/Programs/Horde/Plugins/Build/HordeServer.Build.Tests/Agents/Pools/PoolUpdateServiceTest.cs
@@ -129,6 +129,7 @@ public class PoolUpdateServiceTest : BuildTestSetup
 
 		// Assert
 		Assert.AreEqual(conformRequested, agent.RequestFullConform);
+		Assert.IsFalse(agent.RequestConformSkipWhenDiskFreeSpaceIsAvailable);
 	}
 
 	private async Task<IAgent> CreateAutoConformAgentAsync(int freeDiskSpaceMb, params int?[] autoConformThresholdsMb)
diff --git a/Engine/Source/Programs/Horde/Plugins/Build/HordeServer.Build.Tests/BuildTestSetup.cs b/Engine/Source/Programs/Horde/Plugins/Build/HordeServer.Build.Tests/BuildTestSetup.cs
--- a/Engine/Source/Programs/Horde/Plugins/Build/HordeServer.Build.Tests/BuildTestSetup.cs
+++ b/Engine/Source/Programs/Horde/Plugins/Build/HordeServer.Build.Tests/BuildTestSetup.cs
@@ -93,6 +93,7 @@ namespace HordeServer.Tests
 		public INotificationService NotificationService => ServiceProvider.GetRequiredService<INotificationService>();
 		public IssueService IssueService => ServiceProvider.GetRequiredService<IssueService>();
 		public JobTaskSource JobTaskSource => ServiceProvider.GetRequiredService<JobTaskSource>();
+		public ConformTaskSource ConformTaskSource => ServiceProvider.GetRequiredService<ConformTaskSource>();
 		public JobService JobService => ServiceProvider.GetRequiredService<JobService>();
 		public RpcService RpcService => ServiceProvider.GetRequiredService<RpcService>();
 		public PoolService PoolService => ServiceProvider.GetRequiredService<PoolService>();
diff --git a/Engine/Source/Programs/Horde/Plugins/Build/HordeServer.Build.Tests/Jobs/JobTaskSourceTests.cs b/Engine/Source/Programs/Horde/Plugins/Build/HordeServer.Build.Tests/Jobs/JobTaskSourceTests.cs
--- a/Engine/Source/Programs/Horde/Plugins/Build/HordeServer.Build.Tests/Jobs/JobTaskSourceTests.cs
+++ b/Engine/Source/Programs/Horde/Plugins/Build/HordeServer.Build.Tests/Jobs/JobTaskSourceTests.cs
@@ -8,6 +8,7 @@ using EpicGames.Horde.Jobs.Graphs;
 using EpicGames.Horde.Streams;
 using EpicGames.Horde.Users;
 using HordeCommon.Rpc.Messages;
+using HordeCommon.Rpc.Tasks;
 using HordeServer.Agents;
 using HordeServer.Agents.Pools;
 using HordeServer.Jobs;
@@ -208,13 +209,53 @@ namespace HordeServer.Tests.Jobs
 			if (expectingConform)
 			{
 				Assert.IsTrue(agent.RequestFullConform);
+				Assert.IsTrue(agent.RequestConformSkipWhenDiskFreeSpaceIsAvailable);
 				Assert.IsNull(leaseResult);
 			}
 			else
 			{
 				Assert.IsFalse(agent.RequestFullConform);
+				Assert.IsFalse(agent.RequestConformSkipWhenDiskFreeSpaceIsAvailable);
 				Assert.IsNotNull(leaseResult);
 			}
 		}
+
+		[TestMethod]
+		[DataRow(true, DisplayName = "A Full Conform lease is assigned")]
+		[DataRow(false, DisplayName = "A Non-Full Conform lease assigned")]
+		public async Task AgentAssignConformLeaseAsync(bool fullConform)
+		{
+			Fixture fixture = await CreateFixtureAsync();
+			IAgent agent = fixture.Agent1;
+
+			UpdateAgentOptions options = new() { RequestConformSkipWhenDiskFreeSpaceIsAvailable = true };
+			if (fullConform)
+			{
+				options = options with { RequestFullConform = true };
+			}
+			else
+			{
+				options = options with { RequestConform = true };
+			}
+
+			agent = Deref(await agent.TryUpdateAsync(options));
+
+			Task<CreateLeaseOptions?> lease = Deref(await ConformTaskSource.AssignLeaseAsync(agent, CancellationToken.None));
+			CreateLeaseOptions leaseResult = Deref(await lease.WaitAsync(CancellationToken.None));
+
+			ConformTask? conformTask = leaseResult.Payload as ConformTask;
+			Assert.IsNotNull(conformTask);
+
+			if (fullConform)
+			{
+				Assert.IsTrue(conformTask.SkipWhenDiskFreeSpaceIsAvailable);
+				Assert.IsTrue(conformTask.RemoveUntrackedFiles);
+			}
+			else
+			{
+				Assert.IsFalse(conformTask.SkipWhenDiskFreeSpaceIsAvailable);
+				Assert.IsFalse(conformTask.RemoveUntrackedFiles);
+			}
+		}
 	}
 }
diff --git a/Engine/Source/Programs/Horde/Plugins/Build/HordeServer.Build/Agents/Pools/PoolUpdateService.cs b/Engine/Source/Programs/Horde/Plugins/Build/HordeServer.Build/Agents/Pools/PoolUpdateService.cs
--- a/Engine/Source/Programs/Horde/Plugins/Build/HordeServer.Build/Agents/Pools/PoolUpdateService.cs
+++ b/Engine/Source/Programs/Horde/Plugins/Build/HordeServer.Build/Agents/Pools/PoolUpdateService.cs
@@ -138,9 +138,12 @@ namespace HordeServer.Agents.Pools
 			List<IAgent> agents = (await _agentService.GetCachedAgentsAsync(cancellationToken))
 				.Where(x => x is { Status: AgentStatus.Stopped, Deleted: false, Enabled: true, Ephemeral: false }).ToList();
 
+			// To keep the existing behavior for the scheduled conforms, agent updates will not skip the conform
+			bool skipWhenDiskFreeSpaceIsAvailable = false;
+
 			foreach (IAgent agent in agents)
 			{
-				await ConformUtilities.AutoConformAgentAsync(agent, _poolService, _buildConfig.CurrentValue, _logger, cancellationToken);
+				await ConformUtilities.AutoConformAgentAsync(agent, _poolService, _buildConfig.CurrentValue, skipWhenDiskFreeSpaceIsAvailable, _logger, cancellationToken);
 			}
 		}
 
diff --git a/Engine/Source/Programs/Horde/Plugins/Build/HordeServer.Build/Jobs/JobTaskSource.cs b/Engine/Source/Programs/Horde/Plugins/Build/HordeServer.Build/Jobs/JobTaskSource.cs
--- a/Engine/Source/Programs/Horde/Plugins/Build/HordeServer.Build/Jobs/JobTaskSource.cs
+++ b/Engine/Source/Programs/Horde/Plugins/Build/HordeServer.Build/Jobs/JobTaskSource.cs
@@ -783,7 +783,12 @@ namespace HordeServer.Jobs
 				return true;
 			}
 
-			return await ConformUtilities.AutoConformAgentAsync(agent, _poolService, _buildConfig.CurrentValue, _logger, cancellationToken);
+			// When a conform is scheduled over a job recheck the available disk free space on lease execution.
+			// Agent capabilities are updated asynchronously to lease assignment this means when the conform is executed
+			// the available disk free space on the agent may now be above the conform threshold, which allows the conform to be skipped.
+			bool skipWhenDiskFreeSpaceIsAvailable = true;
+
+			return await ConformUtilities.AutoConformAgentAsync(agent, _poolService, _buildConfig.CurrentValue, skipWhenDiskFreeSpaceIsAvailable, _logger, cancellationToken);
 		}
 
 		/// <summary>
diff --git a/Engine/Source/Programs/Horde/Plugins/Build/HordeServer.Build/Tasks/ConformTaskSource.cs b/Engine/Source/Programs/Horde/Plugins/Build/HordeServer.Build/Tasks/ConformTaskSource.cs
--- a/Engine/Source/Programs/Horde/Plugins/Build/HordeServer.Build/Tasks/ConformTaskSource.cs
+++ b/Engine/Source/Programs/Horde/Plugins/Build/HordeServer.Build/Tasks/ConformTaskSource.cs
@@ -174,6 +174,7 @@ namespace HordeServer.Tasks
 							ILog log = await _logCollection.AddAsync(JobId.Empty, leaseId, agent.SessionId, LogType.Json, cancellationToken: cancellationToken);
 							task.LogId = log.Id.ToString();
 							task.RemoveUntrackedFiles = agent.RequestFullConform;
+							task.SkipWhenDiskFreeSpaceIsAvailable = agent.RequestConformSkipWhenDiskFreeSpaceIsAvailable;
 
 							return LeaseAsync(new CreateLeaseOptions(leaseId, null, "Updating workspaces", null, null, log.Id, null, true, task));
 						}
diff --git a/Engine/Source/Programs/Horde/Plugins/Build/HordeServer.Build/Utilities/ConformUtilities.cs b/Engine/Source/Programs/Horde/Plugins/Build/HordeServer.Build/Utilities/ConformUtilities.cs
--- a/Engine/Source/Programs/Horde/Plugins/Build/HordeServer.Build/Utilities/ConformUtilities.cs
+++ b/Engine/Source/Programs/Horde/Plugins/Build/HordeServer.Build/Utilities/ConformUtilities.cs
@@ -19,9 +19,10 @@ public static class ConformUtilities
 	/// <param name="poolService">Pool service instance</param>
 	/// <param name="config">Current configuration</param>
 	/// <param name="logger">Logger instance</param>
+	/// <param name="skipWhenDiskFreeSpaceIsAvailable">Whether upon execution of the conform should the conform be skipped if the agent now has the available disk free space required</param>
 	/// <param name="cancellationToken">Cancellation token for the async task</param>
 	/// <returns>True if conform was requested for the agent; otherwise false</returns>
-	public static async Task<bool> AutoConformAgentAsync(IAgent agent, PoolService poolService, BuildConfig config, ILogger logger, CancellationToken cancellationToken)
+	public static async Task<bool> AutoConformAgentAsync(IAgent agent, PoolService poolService, BuildConfig config, bool skipWhenDiskFreeSpaceIsAvailable, ILogger logger, CancellationToken cancellationToken)
 	{
 		long MegabytesToBytes(long v) => v * 1024 * 1024;
 		double BytesToMegabytes(long v) => v / 1024.0 / 1024.0;
@@ -43,7 +44,12 @@ public static class ConformUtilities
 		if (freeDiskSpace != null && conformDiskSpaceNeeded > 0 && freeDiskSpace < conformDiskSpaceNeeded &&
 		    agent.ConformAttemptCount is null or 0 && !agent.RequestFullConform)
 		{
-			newAgent = await agent.TryUpdateAsync(new UpdateAgentOptions { RequestFullConform = true }, cancellationToken: cancellationToken);
+			UpdateAgentOptions options = new() { RequestFullConform = true };
+			if (skipWhenDiskFreeSpaceIsAvailable)
+			{
+				options = options with { RequestConformSkipWhenDiskFreeSpaceIsAvailable = true };
+			}
+			newAgent = await agent.TryUpdateAsync(options, cancellationToken);
 			logger.LogInformation("Auto-conforming {AgentId} as workspace conform disk space needed ({ConformDiskSpace:F1} MB) is greater than free disk space ({FreeDiskSpace:F1} MB)",
 				agent.Id.ToString(), BytesToMegabytes(conformDiskSpaceNeeded), BytesToMegabytes(freeDiskSpace.Value));
 		}
diff --git a/Engine/Source/Programs/Horde/Plugins/Compute/HordeServer.Compute/Agents/AgentCollection.cs b/Engine/Source/Programs/Horde/Plugins/Compute/HordeServer.Compute/Agents/AgentCollection.cs
--- a/Engine/Source/Programs/Horde/Plugins/Compute/HordeServer.Compute/Agents/AgentCollection.cs
+++ b/Engine/Source/Programs/Horde/Plugins/Compute/HordeServer.Compute/Agents/AgentCollection.cs
@@ -69,6 +69,7 @@ namespace HordeServer.Agents
 			IReadOnlyList<PoolId> IAgent.ExplicitPools => _document.ExplicitPools;
 			bool IAgent.RequestConform => _document.RequestConform;
 			bool IAgent.RequestFullConform => _document.RequestFullConform;
+			bool IAgent.RequestConformSkipWhenDiskFreeSpaceIsAvailable => _document.RequestConformSkipWhenDiskFreeSpaceIsAvailable;
 			bool IAgent.RequestRestart => _document.RequestRestart;
 			bool IAgent.RequestShutdown => _document.RequestShutdown;
 			bool IAgent.RequestForceRestart => _document.RequestForceRestart;
@@ -189,6 +190,9 @@ namespace HordeServer.Agents
 			[BsonIgnoreIfDefault, BsonDefaultValue(false)]
 			public bool RequestFullConform { get; set; }
 
+			[BsonIgnoreIfDefault, BsonDefaultValue(false)]
+			public bool RequestConformSkipWhenDiskFreeSpaceIsAvailable { get; set; }
+
 			[BsonIgnoreIfDefault, BsonDefaultValue(false)]
 			public bool RequestRestart { get; set; }
 
@@ -762,11 +766,23 @@ namespace HordeServer.Agents
 			{
 				updates.Add(updateBuilder.Set(x => x.RequestConform, options.RequestConform.Value));
 				updates.Add(updateBuilder.Unset(x => x.ConformAttemptCount));
+				if (options.RequestFullConform == null)
+				{
+					updates.Add(updateBuilder.Unset(x => x.RequestConformSkipWhenDiskFreeSpaceIsAvailable));
+				}
 			}
 			if (options.RequestFullConform != null)
 			{
 				updates.Add(updateBuilder.Set(x => x.RequestFullConform, options.RequestFullConform.Value));
 				updates.Add(updateBuilder.Unset(x => x.ConformAttemptCount));
+				if (options.RequestConformSkipWhenDiskFreeSpaceIsAvailable is true)
+				{
+					updates.Add(updateBuilder.Set(x => x.RequestConformSkipWhenDiskFreeSpaceIsAvailable, true));
+				}
+				else
+				{
+					updates.Add(updateBuilder.Unset(x => x.RequestConformSkipWhenDiskFreeSpaceIsAvailable));
+				}
 			}
 			if (options.RequestRestart != null)
 			{
@@ -1071,6 +1087,7 @@ namespace HordeServer.Agents
 			UpdateDefinition<AgentDocument> update = Builders<AgentDocument>.Update.Set(x => x.Workspaces, workspaces);
 			update = update.Set(x => x.LastConformTime, lastConformTime);
 			update = update.Unset(x => x.ConformAttemptCount);
+			update = update.Unset(x => x.RequestConformSkipWhenDiskFreeSpaceIsAvailable);
 			if (!requestConform)
 			{
 				update = update.Unset(x => x.RequestConform);
diff --git a/Engine/Source/Programs/Horde/Plugins/Compute/HordeServer.Compute/Agents/AgentsController.cs b/Engine/Source/Programs/Horde/Plugins/Compute/HordeServer.Compute/Agents/AgentsController.cs
--- a/Engine/Source/Programs/Horde/Plugins/Compute/HordeServer.Compute/Agents/AgentsController.cs
+++ b/Engine/Source/Programs/Horde/Plugins/Compute/HordeServer.Compute/Agents/AgentsController.cs
@@ -235,6 +235,7 @@ namespace HordeServer.Agents
 				agent.Deleted,
 				agent.RequestConform,
 				agent.RequestFullConform,
+				agent.RequestConformSkipWhenDiskFreeSpaceIsAvailable,
 				agent.RequestRestart,
 				agent.RequestShutdown,
 				agent.LastShutdownReason ?? "Unknown",
@@ -287,7 +288,7 @@ namespace HordeServer.Agents
 					return NotFound(agentId);
 				}
 
-				IAgent? newAgent = await agent.TryUpdateAsync(new UpdateAgentOptions(update.Enabled, update.RequestConform, update.RequestFullConform, update.RequestRestart, update.RequestShutdown, update.RequestForceRestart, $"Manual ({userName})", update.Pools?.ConvertAll(x => new PoolId(x)), update.Comment));
+				IAgent? newAgent = await agent.TryUpdateAsync(new UpdateAgentOptions(update.Enabled, update.RequestConform, update.RequestFullConform, update.RequestConformSkipWhenAvailableFreeDiskSpace, update.RequestRestart, update.RequestShutdown, update.RequestForceRestart, $"Manual ({userName})", update.Pools?.ConvertAll(x => new PoolId(x)), update.Comment));
 				if (newAgent == null)
 				{
 					continue;
@@ -306,6 +307,10 @@ namespace HordeServer.Agents
 				{
 					logger.LogInformation("Setting changed: RequestFullConform = {State} ({UserName} via IP {Ip})", newAgent.RequestFullConform, userName, clientIp);
 				}
+				if (agent.RequestConformSkipWhenDiskFreeSpaceIsAvailable != newAgent.RequestConformSkipWhenDiskFreeSpaceIsAvailable)
+				{
+					logger.LogInformation("Setting changed: RequestConformSkipWhenDiskFreeSpaceIsAvailable = {Value} ({UserName} via IP {Ip})", newAgent.RequestConformSkipWhenDiskFreeSpaceIsAvailable, userName, clientIp);
+				}
 				if (agent.RequestRestart != newAgent.RequestRestart)
 				{
 					logger.LogInformation("Setting changed: RequestRestart = {State} ({UserName} via IP {Ip})", newAgent.RequestRestart, userName, clientIp);
diff --git a/Engine/Source/Programs/Horde/Plugins/Compute/HordeServer.Compute/Agents/IAgent.cs b/Engine/Source/Programs/Horde/Plugins/Compute/HordeServer.Compute/Agents/IAgent.cs
--- a/Engine/Source/Programs/Horde/Plugins/Compute/HordeServer.Compute/Agents/IAgent.cs
+++ b/Engine/Source/Programs/Horde/Plugins/Compute/HordeServer.Compute/Agents/IAgent.cs
@@ -167,6 +167,11 @@ namespace HordeServer.Agents
 		/// </summary>
 		public bool RequestFullConform { get; }
 
+		/// <summary>
+		/// Whether a conform can be skipped when the machine already has the required available disk free space needed
+		/// </summary>
+		public bool RequestConformSkipWhenDiskFreeSpaceIsAvailable { get; }
+
 		/// <summary>
 		/// Whether a machine restart is requested
 		/// </summary>
@@ -340,13 +345,14 @@ namespace HordeServer.Agents
 	/// <param name="Enabled">Whether the agent is enabled or not</param>
 	/// <param name="RequestConform">Whether to request a conform job be run</param>
 	/// <param name="RequestFullConform">Whether to request a full conform job be run</param>
+	/// <param name="RequestConformSkipWhenDiskFreeSpaceIsAvailable">Whether a conform can be skipped when the agent already has the available disk free space needed</param>
 	/// <param name="RequestRestart">Whether to request the machine be restarted</param>
 	/// <param name="RequestShutdown">Whether to request the machine be shut down</param>
 	/// <param name="RequestForceRestart">Request an immediate restart without waiting for leases to complete</param>
 	/// <param name="ShutdownReason">The reason for shutting down agent, ex. Autoscaler/Manual/Unexpected</param>
 	/// <param name="ExplicitPools">List of pools for the agent</param>
 	/// <param name="Comment">New comment</param>
-	public record UpdateAgentOptions(bool? Enabled = null, bool? RequestConform = null, bool? RequestFullConform = null, bool? RequestRestart = null, bool? RequestShutdown = null, bool? RequestForceRestart = null, string? ShutdownReason = null, List<PoolId>? ExplicitPools = null, string? Comment = null);
+	public record UpdateAgentOptions(bool? Enabled = null, bool? RequestConform = null, bool? RequestFullConform = null, bool? RequestConformSkipWhenDiskFreeSpaceIsAvailable = null, bool? RequestRestart = null, bool? RequestShutdown = null, bool? RequestForceRestart = null, string? ShutdownReason = null, List<PoolId>? ExplicitPools = null, string? Comment = null);
 
 	/// <summary>
 	/// Options for starting a new agent session
diff --git a/Engine/Source/Programs/Shared/EpicGames.Horde/Agents/AgentMessages.cs b/Engine/Source/Programs/Shared/EpicGames.Horde/Agents/AgentMessages.cs
--- a/Engine/Source/Programs/Shared/EpicGames.Horde/Agents/AgentMessages.cs
+++ b/Engine/Source/Programs/Shared/EpicGames.Horde/Agents/AgentMessages.cs
@@ -51,12 +51,13 @@ namespace EpicGames.Horde.Agents
 	/// <param name="Ephemeral"> Whether the agent is ephemeral </param>
 	/// <param name="RequestConform"> Request a conform be performed using the current agent set </param>
 	/// <param name="RequestFullConform"> Request that a full conform be performed, removing all intermediate files </param>
+	/// <param name="RequestConformSkipWhenAvailableFreeDiskSpace"> Request that a conform can be skipped when the machine already has the available free disk space needed </param>
 	/// <param name="RequestRestart"> Request the machine be restarted </param>
 	/// <param name="RequestShutdown"> Request the machine be shut down </param>
 	/// <param name="RequestForceRestart"> Request the machine be restarted without waiting for leases to complete </param>
 	/// <param name="Pools"> Pools for this agent </param>
 	/// <param name="Comment"> New comment </param>
-	public record UpdateAgentRequest(bool? Enabled = null, bool? Ephemeral = null, bool? RequestConform = null, bool? RequestFullConform = null, bool? RequestRestart = null, bool? RequestShutdown = null, bool? RequestForceRestart = null, List<string>? Pools = null, string? Comment = null);
+	public record UpdateAgentRequest(bool? Enabled = null, bool? Ephemeral = null, bool? RequestConform = null, bool? RequestFullConform = null, bool? RequestConformSkipWhenAvailableFreeDiskSpace = null, bool? RequestRestart = null, bool? RequestShutdown = null, bool? RequestForceRestart = null, List<string>? Pools = null, string? Comment = null);
 
 	/// <summary>
 	/// Response for queries to find a particular lease within an agent
@@ -105,6 +106,7 @@ namespace EpicGames.Horde.Agents
 	/// <param name="Deleted">Whether this agent has expired</param>
 	/// <param name="PendingConform">Whether a conform job is pending</param>
 	/// <param name="PendingFullConform">Whether a full conform job is pending</param>
+	/// <param name="PendingConformSkipWhenDiskFreeSpaceIsAvailable">Whether a conform can be skipped when the agent already has the available disk free space needed</param>
 	/// <param name="PendingRestart">Whether a restart is pending</param>
 	/// <param name="PendingShutdown">Whether a shutdown is pending</param>
 	/// <param name="LastShutdownReason">The reason for the last shutdown</param>
@@ -121,7 +123,7 @@ namespace EpicGames.Horde.Agents
 	/// <param name="Leases">Array of active leases.</param>
 	/// <param name="Workspaces">Current workspaces synced on the agent</param>
 	/// <param name="Comment">Comment for this agent</param>
-	public record GetAgentResponse(AgentId Id, string Name, bool Enabled, AgentStatus Status, double? Rate, SessionId? SessionId, bool Ephemeral, bool Online, bool Deleted, bool PendingConform, bool PendingFullConform, bool PendingRestart, bool PendingShutdown, string LastShutdownReason, DateTime LastConformTime, int? ConformAttemptCount, DateTime? NextConformTime, string? Version, List<string> Properties, Dictionary<string, int> Resources, DateTime? UpdateTime, DateTime? LastOnlineTime, List<string>? Pools, object? Capabilities, List<GetAgentLeaseResponse> Leases, List<GetAgentWorkspaceResponse> Workspaces, string? Comment);
+	public record GetAgentResponse(AgentId Id, string Name, bool Enabled, AgentStatus Status, double? Rate, SessionId? SessionId, bool Ephemeral, bool Online, bool Deleted, bool PendingConform, bool PendingFullConform, bool PendingConformSkipWhenDiskFreeSpaceIsAvailable, bool PendingRestart, bool PendingShutdown, string LastShutdownReason, DateTime LastConformTime, int? ConformAttemptCount, DateTime? NextConformTime, string? Version, List<string> Properties, Dictionary<string, int> Resources, DateTime? UpdateTime, DateTime? LastOnlineTime, List<string>? Pools, object? Capabilities, List<GetAgentLeaseResponse> Leases, List<GetAgentWorkspaceResponse> Workspaces, string? Comment);
 
 	/// <summary>
 	/// Telemetry data for an agent
