Horde: Fix Telemetry Service to use linked CancellationTokenSource for FLTMC command

#Synced-CL 0

diff --git a/Engine/Source/Programs/Horde/HordeAgent/Services/TelemetryService.cs b/Engine/Source/Programs/Horde/HordeAgent/Services/TelemetryService.cs
--- a/Engine/Source/Programs/Horde/HordeAgent/Services/TelemetryService.cs
+++ b/Engine/Source/Programs/Horde/HordeAgent/Services/TelemetryService.cs
@@ -494,9 +494,9 @@ class TelemetryService : IHostedService, IDisposable
 		using ManagedProcess process = new(null, fltmcExePath, "filters", null, null, null, ProcessPriorityClass.Normal);
 		StringBuilder sb = new(1000);
 
-		while (!cancellationToken.IsCancellationRequested)
+		while (!cancellationSource.IsCancellationRequested)
 		{
-			string? line = await process.ReadLineAsync(cancellationToken);
+			string? line = await process.ReadLineAsync(cancellationSource.Token);
 			if (line == null)
 			{
 				break;
@@ -505,7 +505,7 @@ class TelemetryService : IHostedService, IDisposable
 			sb.AppendLine(line);
 		}
 
-		await process.WaitForExitAsync(cancellationToken);
+		await process.WaitForExitAsync(cancellationSource.Token);
 		return sb.ToString();
 	}
 
