[Horde] Expose optional force parameter for the device reservation deletion endpoint
This will allow users to fully delete reservations and associated reservation blocks where appropriate
#tests device unit tests
#Synced-CL 49627308

diff --git a/Engine/Source/Programs/AutomationTool/AutomationUtils/DeviceReservation.cs b/Engine/Source/Programs/AutomationTool/AutomationUtils/DeviceReservation.cs
--- a/Engine/Source/Programs/AutomationTool/AutomationUtils/DeviceReservation.cs
+++ b/Engine/Source/Programs/AutomationTool/AutomationUtils/DeviceReservation.cs
@@ -40,7 +40,15 @@ namespace AutomationTool.DeviceReservation
 		private Reservation ActiveReservation;
 		private List<Device> ReservedDevices;
 
-		// Whether the reservation requires an installation
+		/// <summary>
+		/// When true, the ActiveReservation will be force deleted upon termination of the auto renewal thread
+		/// Force deleted reservations will clear any associated reservation blocks, regardless of the active step
+		/// </summary>
+		public bool ForceDelete { get; set; } = false;
+
+		/// <summary>
+		/// Whether the reservation requires an installation
+		/// <summary>
 		public bool? InstallRequired => ActiveReservation?.InstallRequired;
 
 		public IReadOnlyList<Device> Devices
@@ -131,7 +139,7 @@ namespace AutomationTool.DeviceReservation
 			// Delete reservation on server, if the web request fails backend has logic to cleanup reservations
 			try
 			{
-				ActiveReservation.Delete(ReservationBaseUri);
+				ActiveReservation.Delete(ReservationBaseUri, ForceDelete);
 			}
 			catch (Exception Ex)
 			{
@@ -174,6 +182,13 @@ namespace AutomationTool.DeviceReservation
 			return Builder.Uri;
 		}
 
+		public static Uri AppendQuery(this Uri BaseUri, string NewQuery)
+		{
+			var Builder = new UriBuilder(BaseUri);
+			Builder.Query += NewQuery;
+			return Builder.Uri;
+		}
+
 		public static T InvokeAPI<T>(Uri UriToRequest, string Method, object ObjectToSerialize)
 		{
 			var SerializeOptions = new JsonSerializerOptions() { PropertyNameCaseInsensitive = true };
@@ -531,9 +546,15 @@ namespace AutomationTool.DeviceReservation
 			Utils.InvokeAPI(BaseUri.AppendPath("api/v1/reservations/" + Guid.ToString()), "PUT", NewDuration);
 		}
 
-		public void Delete(Uri BaseUri)
+		public void Delete(Uri BaseUri, bool Force = true)
 		{
-			Utils.InvokeAPI(BaseUri.AppendPath("api/v1/reservations/" + Guid.ToString()), "DELETE");
+			string Path = $"api/v1/reservations/{Guid}";
+			Uri DeleteUri = BaseUri.AppendPath(Path);
+			if (Force)
+			{
+				DeleteUri = DeleteUri.AppendQuery("force=true");
+			}
+			Utils.InvokeAPI(DeleteUri, "DELETE");
 			Utils.Logger.LogInformation("Successfully deleted device reservation \"{GUID}\".", Guid);
 		}
 
diff --git a/Engine/Source/Programs/Horde/Plugins/Build/HordeServer.Build.Tests/Devices/DeviceServiceTests.cs b/Engine/Source/Programs/Horde/Plugins/Build/HordeServer.Build.Tests/Devices/DeviceServiceTests.cs
--- a/Engine/Source/Programs/Horde/Plugins/Build/HordeServer.Build.Tests/Devices/DeviceServiceTests.cs
+++ b/Engine/Source/Programs/Horde/Plugins/Build/HordeServer.Build.Tests/Devices/DeviceServiceTests.cs
@@ -303,6 +303,30 @@ namespace HordeServer.Tests.Devices
 			Assert.AreEqual(telemetry[0].Telemetry[0].JobName, "Test job");
 		}
 
+		[TestMethod]
+		public async Task TestReservationForceDeletionAsync()
+		{
+			await SetupDevicesAsync();
+
+			IJob job = await SetupJobAsync();
+			LegacyCreateReservationRequest request = SetupReservationTestAsync(job, modelId: "Base");
+
+			// create a reservation
+			GetLegacyReservationResponse reservation = ResultToValue(await DeviceController!.CreateDeviceReservationV1Async(request));
+
+			// verify the reservation exists
+			IDeviceReservation? serviceReservation = await DeviceService.TryGetReservationFromLegacyGuidAsync(reservation.Guid);
+			Assert.IsNotNull(serviceReservation);
+
+			// delete the reservation
+			OkResult? deleted = (await DeviceController!.DeleteReservationV1Async(reservation.Guid, true)) as OkResult;
+			Assert.IsNotNull(deleted);
+
+			// Test again to ensure it's now deleted
+			serviceReservation = await DeviceService.TryGetReservationFromLegacyGuidAsync(reservation.Guid);
+			Assert.IsNull(serviceReservation);
+		}
+
 		[TestMethod]
 		public async Task TestReservationNodesAsync()
 		{
@@ -383,6 +407,70 @@ namespace HordeServer.Tests.Devices
 			Assert.AreEqual(telemetry[0].Telemetry[0].JobName, "Test job");
 		}
 
+		[TestMethod]
+		public async Task TestReservationNodeDeletionAsync()
+		{
+			await SetupDevicesAsync();
+			IJob job = await SetupJobAsync("DeviceReserveNodes");
+
+			job = await StartBatchAsync(job, 1);
+			job = await RunStepAsync(job, 1, 0, JobStepOutcome.Success); // Update Version Files
+			job = await RunStepAsync(job, 1, 1, JobStepOutcome.Success); // Compile Editor
+
+			job = await StartBatchAsync(job, 2);
+			job = await RunStepAsync(job, 2, 0, JobStepOutcome.Success); // Compile Client
+
+			job = await StartBatchAsync(job, 3);
+			job = await RunStepAsync(job, 3, 0, JobStepOutcome.Success); // Cook Client
+
+			job = await StartBatchAsync(job, 4);
+
+			// Install  the build
+			JobStepId stepId = GetStepId(job, "Install Build");
+			job = await StartStepAsync(job, 4, 0); // Install Build
+			LegacyCreateReservationRequest request = SetupReservationTestAsync(job, stepId: stepId);
+			GetLegacyReservationResponse installReservation = ResultToValue(await DeviceController!.CreateDeviceReservationV1Async(request));
+			job = await FinishStepAsync(job, 4, 0, JobStepOutcome.Success); // Install Build
+
+			for (int i = 1; i < 3; i++)
+			{
+				// Run Test 1
+				stepId = GetStepId(job, $"Run Test {i}");
+				job = await StartStepAsync(job, 4, i);
+
+				request = SetupReservationTestAsync(job, stepId: stepId);
+
+				ActionResult<GetLegacyReservationResponse> result = await DeviceController!.CreateDeviceReservationV1Async(request);
+
+				GetLegacyReservationResponse? reservation = ResultToValue(result);
+				Assert.IsNotNull(reservation);
+				Assert.AreEqual(installReservation.Guid, reservation.Guid);
+
+				if (i == 1)
+				{
+					// For the first step, simply delete the reservation at the end of a test like normal
+					await DeviceController!.DeleteReservationV1Async(reservation.Guid);
+					job = await FinishStepAsync(job, 4, i, JobStepOutcome.Success);
+					await DeviceService.TickForTestingAsync();
+
+					// We should find that the reservation still exists because it's part of a reservation block
+					IDeviceReservation? serviceReservation = await DeviceService.TryGetReservationFromLegacyGuidAsync(reservation.Guid);
+					Assert.IsNotNull(serviceReservation);
+				}
+				else if (i == 2)
+				{
+					// For the second step, force delete the reservation at the end of a test to verify the block is cleared despite there being remaining steps
+					await DeviceController!.DeleteReservationV1Async(reservation.Guid, true);
+					job = await FinishStepAsync(job, 4, i, JobStepOutcome.Success);
+					await DeviceService.TickForTestingAsync();
+
+					// We should find that the reservation no longer exists because it was force deleted
+					IDeviceReservation? serviceReservation = await DeviceService.TryGetReservationFromLegacyGuidAsync(reservation.Guid);
+					Assert.IsNull(serviceReservation);
+				}
+			}
+		}
+
 		[TestMethod]
 		public async Task TestReservationMarkersAsync()
 		{
diff --git a/Engine/Source/Programs/Horde/Plugins/Build/HordeServer.Build/Devices/DeviceService.cs b/Engine/Source/Programs/Horde/Plugins/Build/HordeServer.Build/Devices/DeviceService.cs
--- a/Engine/Source/Programs/Horde/Plugins/Build/HordeServer.Build/Devices/DeviceService.cs
+++ b/Engine/Source/Programs/Horde/Plugins/Build/HordeServer.Build/Devices/DeviceService.cs
@@ -609,16 +609,20 @@ namespace HordeServer.Devices
 		/// </summary>
 		public async Task<bool> DeleteReservationAsync(ObjectId id, bool force = false)
 		{
+			IDeviceReservation? reservation = await _devices.TryGetReservationAsync(id);
+
 			if (!force)
 			{
-				IDeviceReservation? reservation = await _devices.TryGetReservationAsync(id);
-
 				if (reservation != null && reservation.JobId != null && reservation.ReservedStepIds != null && reservation.ReservedStepIds.Count > 0)
 				{
 					// will be deleted in tick
 					return false;
 				}
 			}
+			else if (reservation != null)
+			{
+				_logger.LogInformation("Force deleting reservation {ReservationGuid}", reservation.LegacyGuid);
+			}
 
 			return await _devices.DeleteReservationAsync(id);
 		}
diff --git a/Engine/Source/Programs/Horde/Plugins/Build/HordeServer.Build/Devices/DevicesController.cs b/Engine/Source/Programs/Horde/Plugins/Build/HordeServer.Build/Devices/DevicesController.cs
--- a/Engine/Source/Programs/Horde/Plugins/Build/HordeServer.Build/Devices/DevicesController.cs
+++ b/Engine/Source/Programs/Horde/Plugins/Build/HordeServer.Build/Devices/DevicesController.cs
@@ -993,7 +993,7 @@ namespace HordeServer.Devices
 		/// </summary>
 		[HttpDelete]
 		[Route("/api/v1/reservations/{reservationGuid}")]
-		public async Task<ActionResult> DeleteReservationV1Async(string reservationGuid)
+		public async Task<ActionResult> DeleteReservationV1Async(string reservationGuid, [FromQuery] bool force = false)
 		{
 			IDeviceReservation? reservation = await _deviceService.TryGetReservationFromLegacyGuidAsync(reservationGuid);
 
@@ -1002,7 +1002,7 @@ namespace HordeServer.Devices
 				return BadRequest($"Unable to find reservation for guid {reservationGuid}");
 			}
 
-			await _deviceService.DeleteReservationAsync(reservation.Id);
+			await _deviceService.DeleteReservationAsync(reservation.Id, force);
 
 			return Ok();
 		}
