Fix DotNet tasks not using overridden DotNetPath

#rnx
#Synced-CL 49686086

diff --git a/Engine/Source/Programs/Horde/BuildHorde.xml b/Engine/Source/Programs/Horde/BuildHorde.xml
--- a/Engine/Source/Programs/Horde/BuildHorde.xml
+++ b/Engine/Source/Programs/Horde/BuildHorde.xml
@@ -153,8 +153,8 @@
 		</Node>
 		<Node Name="Build HordeAgent" Produces="#HordeAgentFiles">
 			<Delete Files="$(ServerToolsStagingDir_HordeAgent)/..."/>
-			<DotNet Arguments="publish &quot;$(RootDir)/Engine/Source/Programs/Horde/HordeAgent/HordeAgent.csproj&quot; --output &quot;$(ServerToolsStagingDir_HordeAgent)&quot; $(VersionArguments)"/>
-			<DotNet Arguments="publish &quot;$(RootDir)/Engine/Source/Programs/Horde/Drivers/JobDriver/JobDriver.csproj&quot; --output &quot;$(ServerToolsStagingDir_HordeAgent)/JobDriver&quot; $(VersionArguments)"/>
+			<DotNet Arguments="publish &quot;$(RootDir)/Engine/Source/Programs/Horde/HordeAgent/HordeAgent.csproj&quot; --output &quot;$(ServerToolsStagingDir_HordeAgent)&quot; $(VersionArguments)" DotNetPath="$(DotNetPath)"/>
+			<DotNet Arguments="publish &quot;$(RootDir)/Engine/Source/Programs/Horde/Drivers/JobDriver/JobDriver.csproj&quot; --output &quot;$(ServerToolsStagingDir_HordeAgent)/JobDriver&quot; $(VersionArguments)" DotNetPath="$(DotNetPath)"/>
 			<Delete Files="$(ServerToolsStagingDir_HordeAgent)/appsettings.Local.json"/>
 			<Sign Files="$(ServerToolsStagingDir_HordeAgent)/HordeAgent.exe;$(ServerToolsStagingDir_HordeAgent)/HordeAgent.dll;$(ServerToolsStagingDir_HordeAgent)/EpicGames.*.dll" If="'$(SignOutputs)' == 'true'"/>
 			<Sign Files="$(ServerToolsStagingDir_HordeAgent)/JobDriver/JobDriver.exe;$(ServerToolsStagingDir_HordeAgent)/JobDriver/JobDriver.dll;$(ServerToolsStagingDir_HordeAgent)/JobDriver/EpicGames.*.dll" If="'$(SignOutputs)' == 'true'"/>
@@ -164,11 +164,11 @@
 			<Delete Files="$(ServerToolsStagingDir_HordeAgentWin64)/..."/>
 			
 			<!-- Build the agent -->
-			<DotNet Arguments="publish &quot;$(RootDir)/Engine/Source/Programs/Horde/HordeAgent/HordeAgent.csproj&quot; --output &quot;$(ServerToolsStagingDir_HordeAgentWin64)&quot; --runtime win-x64 --self-contained $(VersionArguments)"/>
+			<DotNet Arguments="publish &quot;$(RootDir)/Engine/Source/Programs/Horde/HordeAgent/HordeAgent.csproj&quot; --output &quot;$(ServerToolsStagingDir_HordeAgentWin64)&quot; --runtime win-x64 --self-contained $(VersionArguments)" DotNetPath="$(DotNetPath)"/>
 			<Sign Files="$(ServerToolsStagingDir_HordeAgentWin64)/HordeAgent.exe;$(ServerToolsStagingDir_HordeAgentWin64)/HordeAgent.dll;$(ServerToolsStagingDir_HordeAgentWin64)/EpicGames.*.dll" If="'$(SignOutputs)' == 'true'"/>
 
 			<!-- Build the driver -->
-			<DotNet Arguments="publish &quot;$(RootDir)/Engine/Source/Programs/Horde/Drivers/JobDriver/JobDriver.csproj&quot; --output &quot;$(ServerToolsStagingDir_HordeAgentWin64)/JobDriver&quot; --runtime win-x64 --self-contained $(VersionArguments)"/>
+			<DotNet Arguments="publish &quot;$(RootDir)/Engine/Source/Programs/Horde/Drivers/JobDriver/JobDriver.csproj&quot; --output &quot;$(ServerToolsStagingDir_HordeAgentWin64)/JobDriver&quot; --runtime win-x64 --self-contained $(VersionArguments)" DotNetPath="$(DotNetPath)"/>
 			<Sign Files="$(ServerToolsStagingDir_HordeAgentWin64)/JobDriver/JobDriver.exe;$(ServerToolsStagingDir_HordeAgentWin64)/JobDriver/JobDriver.dll;$(ServerToolsStagingDir_HordeAgentWin64)/JobDriver/EpicGames.*.dll" If="'$(SignOutputs)' == 'true'"/>
 
 			<Delete Files="$(ServerToolsStagingDir_HordeAgentWin64)/appsettings.Local.json"/>
@@ -206,21 +206,21 @@
 		</Node>
 		<Node Name="Build P4VUtils" Produces="#P4VUtilsFiles">
 			<Delete Files="$(ServerToolsStagingDir_P4VUtils)/..."/>
-			<DotNet Arguments="publish &quot;$(RootDir)/Engine/Extras/P4VUtils/P4VUtils.csproj&quot; --configuration Release --runtime win-x64 --self-contained --output &quot;$(ServerToolsStagingDir_P4VUtils)&quot; /p:IsWindows=true /p:IsOSX=false /p:IsLinux=false /p:WithRestricted=false $(VersionArguments)"/>
+			<DotNet Arguments="publish &quot;$(RootDir)/Engine/Extras/P4VUtils/P4VUtils.csproj&quot; --configuration Release --runtime win-x64 --self-contained --output &quot;$(ServerToolsStagingDir_P4VUtils)&quot; /p:IsWindows=true /p:IsOSX=false /p:IsLinux=false /p:WithRestricted=false $(VersionArguments)" DotNetPath="$(DotNetPath)"/>
 			<Tag Files="$(ServerToolsStagingDir_P4VUtils)/..." With="#P4VUtilsFiles"/>
 		</Node>
 		<Node Name="Build UnrealGameSync" Produces="#UnrealGameSyncFiles">
 			<Delete Files="$(ServerToolsStagingDir_UnrealGameSync)/..."/>
-			<DotNet Arguments="publish &quot;$(UgsDir)/UnrealGameSyncCmd/UnrealGameSyncCmd.csproj&quot; --output=&quot;$(ServerToolsStagingDir_UnrealGameSync)&quot; --self-contained -c ReleaseAutoUpdate -r win-x64 $(VersionArguments)"/>
-			<DotNet Arguments="publish &quot;$(UgsDir)/UnrealGameSync/UnrealGameSync.csproj&quot; --output=&quot;$(ServerToolsStagingDir_UnrealGameSync)&quot; --self-contained -c ReleaseAutoUpdate -r win-x64 $(VersionArguments)"/>
+			<DotNet Arguments="publish &quot;$(UgsDir)/UnrealGameSyncCmd/UnrealGameSyncCmd.csproj&quot; --output=&quot;$(ServerToolsStagingDir_UnrealGameSync)&quot; --self-contained -c ReleaseAutoUpdate -r win-x64 $(VersionArguments)" DotNetPath="$(DotNetPath)"/>
+			<DotNet Arguments="publish &quot;$(UgsDir)/UnrealGameSync/UnrealGameSync.csproj&quot; --output=&quot;$(ServerToolsStagingDir_UnrealGameSync)&quot; --self-contained -c ReleaseAutoUpdate -r win-x64 $(VersionArguments)" DotNetPath="$(DotNetPath)"/>
 			<Sign Files="$(ServerToolsStagingDir_UnrealGameSync)/UnrealGameSync.exe;$(ServerToolsStagingDir_UnrealGameSync)/ugs.exe;$(ServerToolsStagingDir_UnrealGameSync)/ugs.dll;$(ServerToolsStagingDir_UnrealGameSync)/UnrealGameSync*.dll;$(ServerToolsStagingDir_UnrealGameSync)/EpicGames.*.dll" If="'$(SignOutputs)' == 'true'"/>
 			<Tag Files="$(ServerToolsStagingDir_UnrealGameSync)/..." With="#UnrealGameSyncFiles"/>
 		</Node>
 		<Node Name="Build UnrealGameSync Cmd" Produces="#UnrealGameSyncCmdFiles">
 			<Delete Files="$(ServerToolsStagingDir_UnrealGameSyncCmd)/..."/>
-			<DotNet Arguments="publish &quot;$(UgsDir)/UnrealGameSyncCmd/UnrealGameSyncCmd.csproj&quot; --output=&quot;$(ServerToolsStagingDir_UnrealGameSyncCmd)&quot; --no-self-contained -c Release -r osx-x64 $(VersionArguments)"/>
-			<DotNet Arguments="publish &quot;$(UgsDir)/UnrealGameSyncCmd/UnrealGameSyncCmd.csproj&quot; --output=&quot;$(ServerToolsStagingDir_UnrealGameSyncCmd)&quot; --no-self-contained -c Release -r linux-x64 $(VersionArguments)"/>
-			<DotNet Arguments="publish &quot;$(UgsDir)/UnrealGameSyncCmd/UnrealGameSyncCmd.csproj&quot; --output=&quot;$(ServerToolsStagingDir_UnrealGameSyncCmd)&quot; --no-self-contained -c Release -r win-x64 $(VersionArguments)"/>
+			<DotNet Arguments="publish &quot;$(UgsDir)/UnrealGameSyncCmd/UnrealGameSyncCmd.csproj&quot; --output=&quot;$(ServerToolsStagingDir_UnrealGameSyncCmd)&quot; --no-self-contained -c Release -r osx-x64 $(VersionArguments)" DotNetPath="$(DotNetPath)"/>
+			<DotNet Arguments="publish &quot;$(UgsDir)/UnrealGameSyncCmd/UnrealGameSyncCmd.csproj&quot; --output=&quot;$(ServerToolsStagingDir_UnrealGameSyncCmd)&quot; --no-self-contained -c Release -r linux-x64 $(VersionArguments)" DotNetPath="$(DotNetPath)"/>
+			<DotNet Arguments="publish &quot;$(UgsDir)/UnrealGameSyncCmd/UnrealGameSyncCmd.csproj&quot; --output=&quot;$(ServerToolsStagingDir_UnrealGameSyncCmd)&quot; --no-self-contained -c Release -r win-x64 $(VersionArguments)" DotNetPath="$(DotNetPath)"/>
 			<Sign Files="$(ServerToolsStagingDir_UnrealGameSyncCmd)/ugs.dll;$(ServerToolsStagingDir_UnrealGameSyncCmd)/UnrealGameSync*.dll;$(ServerToolsStagingDir_UnrealGameSyncCmd)/EpicGames.*.dll" If="'$(SignOutputs)' == 'true'"/>
 			<Tag Files="$(ServerToolsStagingDir_UnrealGameSyncCmd)/..." With="#UnrealGameSyncCmdFiles"/>
 		</Node>
