Add successful test phase with warning filtering

#rnx

#Synced-CL 49762367

diff --git a/Engine/Source/Programs/Horde/HordeDashboard/plugins/test-automation/api.ts b/Engine/Source/Programs/Horde/HordeDashboard/plugins/test-automation/api.ts
--- a/Engine/Source/Programs/Horde/HordeDashboard/plugins/test-automation/api.ts
+++ b/Engine/Source/Programs/Horde/HordeDashboard/plugins/test-automation/api.ts
@@ -1,7 +1,7 @@
 // Copyright Epic Games, Inc. All Rights Reserved.
 
 import backend from "horde/backend";
-import { GetUserResponse } from "horde/backend/Api";
+import { AuditLogEntry, AuditLogQuery, GetUserResponse } from "horde/backend/Api";
 
 /// A test data response
 export type GetTestDataResponse = {
@@ -230,6 +230,9 @@ export enum TestPhaseOutcome {
     /// The test phase was skipped
     Skipped = "Skipped",
 
+    /// The test phase was successful but with warning
+    Warning = "Warning",
+
     /// The test phase outcome is unknown
     Unknown = "Unknown"
 }
@@ -743,6 +746,15 @@ export async function getTestAuditsV2(testIds: string[], minCreateTime?: string)
     });
 }
 
+export async function getTestAuditLogHistoryV2(id: string, query: AuditLogQuery): Promise<AuditLogEntry[]> {
+    return new Promise<AuditLogEntry[]>((resolve, reject) => {
+        backend.fetch.get(`/api/v3/testdata/audit/log/${encodeURIComponent(id)}`, { params: query }).then((response) => {
+            const history = (response.data?.entries ?? []) as AuditLogEntry[];
+            resolve(history);
+        }).catch(reason => { reject(reason); });
+    });
+}
+
 export async function getUsersInfo(userIds: string[]): Promise<GetUserResponse[]> {
     const results = await backend.getUsers({
         ids: userIds,
diff --git a/Engine/Source/Programs/Horde/HordeDashboard/plugins/test-automation/testAutomationCommon.tsx b/Engine/Source/Programs/Horde/HordeDashboard/plugins/test-automation/testAutomationCommon.tsx
--- a/Engine/Source/Programs/Horde/HordeDashboard/plugins/test-automation/testAutomationCommon.tsx
+++ b/Engine/Source/Programs/Horde/HordeDashboard/plugins/test-automation/testAutomationCommon.tsx
@@ -138,8 +138,9 @@ export const getStatusColors = (): Map<StatusColor, string> => {
     return colors;
 }
 
-export const statusTexts = new Map<TestOutcome, string>([
+export const statusTexts = new Map<string, string>([
     [TestOutcome.Success, "Passed"],
+    [TestPhaseOutcome.Warning, "With Warning"],
     [TestOutcome.Failure, "Failure"],
     [TestOutcome.Skipped, "Skipped"],
     [TestOutcome.Unspecified, "Catastrophic"]
@@ -515,25 +516,31 @@ export const getPhaseOutcomeOptions = (): ICheckListOption[] => {
     const statusColors = getStatusColors();
     return [
         {
-            key: TestOutcome.Unspecified,
+            key: TestPhaseOutcome.Interrupted,
             text: statusTexts.get(TestOutcome.Unspecified),
             checked: true,
             color: statusColors.get(StatusColor.Unspecified)
         },
         {
-            key: TestOutcome.Failure,
+            key: TestPhaseOutcome.Failed,
             text: statusTexts.get(TestOutcome.Failure),
             checked: true,
             color: statusColors.get(StatusColor.Failure)
         },
         {
-            key: TestOutcome.Success,
+            key: TestPhaseOutcome.Success,
             text: statusTexts.get(TestOutcome.Success),
             checked: true,
             color: statusColors.get(StatusColor.Success)
         },
         {
-            key: TestOutcome.Skipped,
+            key: TestPhaseOutcome.Warning,
+            text: statusTexts.get(TestPhaseOutcome.Warning),
+            checked: true,
+            color: statusColors.get(StatusColor.Warnings)
+        },
+        {
+            key: TestPhaseOutcome.Skipped,
             text: statusTexts.get(TestOutcome.Skipped),
             checked: true,
             color: statusColors.get(StatusColor.Skipped)
@@ -563,19 +570,22 @@ class PhaseFiltersContext {
         const outcome: Set<TestPhaseOutcome> = new Set();
         this._outcomeOptions.filter(option => option.checked).forEach(option => {
             switch(option.key) {
-                case TestOutcome.Unspecified:
+                case TestPhaseOutcome.Interrupted:
                     outcome.add(TestPhaseOutcome.NotRun);
                     outcome.add(TestPhaseOutcome.Unknown);
                     outcome.add(TestPhaseOutcome.Interrupted);
                     break;
-                case TestOutcome.Skipped:
+                case TestPhaseOutcome.Skipped:
                     outcome.add(TestPhaseOutcome.Skipped);
                     break;
-                case TestOutcome.Failure:
+                case TestPhaseOutcome.Failed:
                     outcome.add(TestPhaseOutcome.Failed);
                     outcome.add(TestPhaseOutcome.Interrupted);
                     break;
-                case TestOutcome.Success:
+                case TestPhaseOutcome.Warning:
+                    outcome.add(TestPhaseOutcome.Warning);
+                    break;
+                case TestPhaseOutcome.Success:
                     outcome.add(TestPhaseOutcome.Success);
                     break;
             }
diff --git a/Engine/Source/Programs/Horde/HordeDashboard/plugins/test-automation/testAutomationPhaseView.tsx b/Engine/Source/Programs/Horde/HordeDashboard/plugins/test-automation/testAutomationPhaseView.tsx
--- a/Engine/Source/Programs/Horde/HordeDashboard/plugins/test-automation/testAutomationPhaseView.tsx
+++ b/Engine/Source/Programs/Horde/HordeDashboard/plugins/test-automation/testAutomationPhaseView.tsx
@@ -347,7 +347,7 @@ export const TestPhasesView: React.FC<{ testDataId: string, handler: TestDataHan
         interruptedCursor.current.setCursor(selectedPhase.current?.key);
         phaseCursor.current = new CircularList(phases.map((p) => p.key));
         phaseCursor.current.setCursor(selectedPhase.current?.key);
-        const warnings = phases.filter((p) => p.outcome !== TestPhaseOutcome.Success && p.hasWarning).map((p) => p.key) ?? [];
+        const warnings = phases.filter((p) => p.outcome === TestPhaseOutcome.Success && p.hasWarning).map((p) => p.key) ?? [];
         warningCursor.current = new CircularList(warnings);
         warningCursor.current.setCursor(selectedPhase.current?.key);
         setFilteredPhases(phases);
diff --git a/Engine/Source/Programs/Horde/HordeDashboard/plugins/test-automation/testData.ts b/Engine/Source/Programs/Horde/HordeDashboard/plugins/test-automation/testData.ts
--- a/Engine/Source/Programs/Horde/HordeDashboard/plugins/test-automation/testData.ts
+++ b/Engine/Source/Programs/Horde/HordeDashboard/plugins/test-automation/testData.ts
@@ -315,7 +315,9 @@ export abstract class TestPhaseStatus {
     }
 
     isMatch(names: string[], outcomes?: Set<TestPhaseOutcome>, tags?: string[]) {
-        if (!!outcomes && !outcomes.has(this.outcome)) return false;
+        if (!!outcomes && !outcomes.has(this.outcome)
+            && (!outcomes.has(TestPhaseOutcome.Warning) || this.outcome !== TestPhaseOutcome.Success || !this.hasWarning)) /** allow filtering of success but with warning */
+                return false;
         const lowerTags = this.tags?.map(t => t.toLowerCase());
         if (!!tags && (!lowerTags || !lowerTags.some((t) => tags.includes(t)))) return false;
         if (names.length === 0) return true;
