Horde: Fix conform in PerforceMaterializer to run cleanup after sync and improve error handling

diff --git a/Engine/Source/Programs/Horde/Drivers/JobDriver.Tests/PerforceMaterializerTests.cs b/Engine/Source/Programs/Horde/Drivers/JobDriver.Tests/PerforceMaterializerTests.cs
--- a/Engine/Source/Programs/Horde/Drivers/JobDriver.Tests/PerforceMaterializerTests.cs
+++ b/Engine/Source/Programs/Horde/Drivers/JobDriver.Tests/PerforceMaterializerTests.cs
@@ -216,6 +216,28 @@ public class PerforceMaterializerTests : BasePerforceFixtureTest
 		
 		await PerforceExecutor.ConformMaterializersAsync(factories, TempDir, [raw], removeUntrackedFiles, [], NoOpTracer, logger, CancellationToken.None);
 	}
+
+	[TestMethod]
+	public async Task Conform_WithDeletedClient_ShouldStillClean_Async()
+	{
+		// Arrange: Sync first, then delete the client, simulating a stale/deleted client
+		using PerforceMaterializer pm = CreatePm(PerforceSnapshot.FooMain);
+		await pm.SyncAsync(2, -1, new SyncOptions(), CancellationToken.None);
+		await pm.FinalizeAsync(CancellationToken.None);
+
+		// Delete the client to simulate a stale client scenario
+		await pm.DeleteClientForTestAsync(CancellationToken.None);
+		
+		// Pollute the workspace with an untracked file
+		await File.WriteAllTextAsync(Path.Join(pm.SyncDir.FullName, _untrackedFile.clientFile), _untrackedFile.content);
+
+		// Act: Conform with removeUntrackedFiles=true (which triggers CleanAsync)
+		// This should recreate the client and successfully clean
+		await pm.ConformAsync(true, CancellationToken.None);
+
+		// Assert: Workspace should be valid
+		PerforceSnapshot.GetSnapshot().GetStream(PerforceSnapshot.FooMain).GetWorkspace(9).AssertDepotFiles(pm.SyncDir.FullName);
+	}
 	
 	[TestMethod]
 	[DataRow("Platforms/...", "//Foo/Main/Platforms/... //MyComputer/Platforms/...")]
diff --git a/Engine/Source/Programs/Horde/Drivers/JobDriver/Execution/PerforceMaterializer.cs b/Engine/Source/Programs/Horde/Drivers/JobDriver/Execution/PerforceMaterializer.cs
--- a/Engine/Source/Programs/Horde/Drivers/JobDriver/Execution/PerforceMaterializer.cs
+++ b/Engine/Source/Programs/Horde/Drivers/JobDriver/Execution/PerforceMaterializer.cs
@@ -415,17 +415,23 @@ public sealed class PerforceMaterializer : IWorkspaceMaterializer
 		using ILoggerProgress status = _logger.BeginProgressScope("Conforming...");
 		
 		await InitializeAsync(cancellationToken);
-		IPerforceConnection perforceWithClient = await GetPerforceWithClientAsync();
-
-		if (removeUntrackedFiles)
-		{
-			await CleanAsync(perforceWithClient, cancellationToken);
-		}
 
 		// Sync from scratch again. Any files are reverted by SyncAsync
 		// TODO: Call ResetState() to ensure a complete sync is performed?
 		SyncOptions syncOptions = new () { RemoveUntracked = removeUntrackedFiles };
 		await SyncAsync(IWorkspaceMaterializer.LatestChangeNumber, 0, syncOptions, cancellationToken);
+		
+		if (removeUntrackedFiles)
+		{
+			if (_perforceWithClient != null)
+			{
+				await CleanAsync(_perforceWithClient, cancellationToken);
+			}
+			else
+			{
+				_logger.LogWarning("Cannot clean: Perforce client not initialized after sync");
+			}
+		}
 	}
 	
 		/// <summary>
@@ -679,7 +685,16 @@ public sealed class PerforceMaterializer : IWorkspaceMaterializer
 		}
 		else
 		{
-			_logger.LogInformation("Clean skipped");
+			PerforceError? error = res.FirstError;
+			if (error != null)
+			{
+				_logger.LogWarning("Clean failed: {Error}", error.ToString());
+				span.SetAttribute("horde.pm.clean_error", error.Data ?? "Unknown error");
+			}
+			else
+			{
+				_logger.LogWarning("Clean failed with no error details");
+			}
 		}
 		
 		await SaveStateAsync(TransactionStatus.Clean, cancellationToken);
