Horde: Support Jira Cloud and On Premise for Build Health

diff --git a/Engine/Source/Programs/Horde/Plugins/Build/HordeServer.Build/Issues/External/JiraService.cs b/Engine/Source/Programs/Horde/Plugins/Build/HordeServer.Build/Issues/External/JiraService.cs
--- a/Engine/Source/Programs/Horde/Plugins/Build/HordeServer.Build/Issues/External/JiraService.cs
+++ b/Engine/Source/Programs/Horde/Plugins/Build/HordeServer.Build/Issues/External/JiraService.cs
@@ -91,6 +91,16 @@ namespace HordeServer.Issues.External
 	/// </summary>
 	internal sealed class JiraService : IExternalIssueService, IHostedService, IAsyncDisposable
 	{
+		/// <summary>
+		/// Cached JsonSerializerOptions for case-insensitive deserialization
+		/// </summary>
+		static readonly JsonSerializerOptions s_jsonOptions = new JsonSerializerOptions { PropertyNameCaseInsensitive = true };
+
+		/// <summary>
+		/// Normalizes a URL by removing any trailing slashes
+		/// </summary>
+		static Uri NormalizeUrl(Uri url) => new Uri(url.ToString().TrimEnd('/'));
+
 		readonly ILogger _logger;
 
 		/// <summary>
@@ -113,6 +123,27 @@ namespace HordeServer.Issues.External
 
 		readonly Uri _jiraUrl;
 
+		/// <summary>
+		/// Base URL returned from Jira's serverInfo endpoint, used for all REST calls.
+		/// Initialized to configured URL, updated after serverInfo detection.
+		/// </summary>
+		Uri _jiraBaseUrl;
+
+		// Ticker timing configuration
+		readonly TimeSpan _tickInitialDelay = TimeSpan.FromSeconds(15);
+		readonly TimeSpan _tickInterval = TimeSpan.FromMinutes(2);
+		readonly TimeSpan _updateTickInterval = TimeSpan.FromMinutes(1);
+
+		/// <summary>
+		/// Detected Jira deployment type (Cloud vs Server/Data Center)
+		/// </summary>
+		JiraDeploymentType _deploymentType = JiraDeploymentType.Unknown;
+
+		/// <summary>
+		/// Returns true if connected to Jira Cloud, false for Server/Data Center
+		/// </summary>
+		bool IsCloud => _deploymentType == JiraDeploymentType.Cloud;
+
 		readonly Dictionary<string, JiraProject> _cachedJiraProjects = new Dictionary<string, JiraProject>();
 
 		readonly IOptionsMonitor<BuildConfig> _buildConfig;
@@ -132,9 +163,10 @@ namespace HordeServer.Issues.External
 			_logger = logger;
 			_issueService = issueService;
 			_userCollection = userCollection;
-			_ticker = clock.AddTicker<JiraService>(TimeSpan.FromMinutes(2.0), TickAsync, logger);
-			_updateTicker = clock.AddSharedTicker("JiraService.Update", TimeSpan.FromMinutes(1.0), TickUpdateAsync, logger);
-			_jiraUrl = _settings.JiraUrl!;
+			_ticker = clock.AddTicker<JiraService>(_tickInitialDelay, TickAsync, logger);
+			_updateTicker = clock.AddSharedTicker("JiraService.Update", _updateTickInterval, TickUpdateAsync, logger);
+			_jiraUrl = NormalizeUrl(_settings.JiraUrl!);
+			_jiraBaseUrl = _jiraUrl;
 			_buildConfig = buildConfig;
 
 			// setup http client for Jira rest api queries
@@ -148,10 +180,42 @@ namespace HordeServer.Issues.External
 		/// <inheritdoc/>
 		public async Task StartAsync(CancellationToken cancellationToken)
 		{
+			await DetectDeploymentTypeAsync(cancellationToken);
 			await _ticker.StartAsync();
 			await _updateTicker.StartAsync();
 		}
 
+		/// <summary>
+		/// Detects whether the Jira instance is Cloud or Server/Data Center by querying the serverInfo endpoint
+		/// </summary>
+		async Task DetectDeploymentTypeAsync(CancellationToken cancellationToken)
+		{
+			try
+			{
+				Uri uri = new Uri(_jiraUrl, "/rest/api/2/serverInfo");
+				HttpResponseMessage response = await _retryPolicy.ExecuteAsync(ctx => _client.GetAsync(uri, ctx), cancellationToken);
+				response.EnsureSuccessStatusCode();
+
+				string responseBody = await response.Content.ReadAsStringAsync(cancellationToken);
+				ServerInfoResponse? serverInfo = JsonSerializer.Deserialize<ServerInfoResponse>(responseBody, s_jsonOptions);
+
+				bool isCloud = String.Equals(serverInfo?.DeploymentType, "Cloud", StringComparison.OrdinalIgnoreCase);
+				_deploymentType = isCloud ? JiraDeploymentType.Cloud : JiraDeploymentType.OnPremise;
+
+				if (!String.IsNullOrEmpty(serverInfo?.BaseUrl))
+				{
+					_jiraBaseUrl = NormalizeUrl(new Uri(serverInfo.BaseUrl));
+				}
+
+				_logger.LogInformation("Detected Jira {DeploymentType} (version: {Version}, baseUrl: {BaseUrl})", serverInfo?.DeploymentType ?? "Unknown", serverInfo?.Version, _jiraBaseUrl);
+			}
+			catch (Exception ex) when (ex is not OperationCanceledException)
+			{
+				_logger.LogWarning(ex, "Unable to detect Jira deployment type, defaulting to on-premise");
+				_deploymentType = JiraDeploymentType.OnPremise;
+			}
+		}
+
 		/// <inheritdoc/>
 		public async Task StopAsync(CancellationToken cancellationToken)
 		{
@@ -168,7 +232,7 @@ namespace HordeServer.Issues.External
 		}
 
 		/// <inheritdoc/>
-		public string? GetIssueUrl(string key) => $"{_jiraUrl}browse/{key}";
+		public string? GetIssueUrl(string key) => $"{_jiraBaseUrl}/browse/{key}";
 
 		public Task<List<IExternalIssueProject>> GetProjectsAsync(StreamConfig streamConfig, CancellationToken cancellationToken)
 		{
@@ -210,7 +274,7 @@ namespace HordeServer.Issues.External
 			{
 				string projectKey = jiraProjectKeys[i];
 
-				Uri uri = new Uri(_jiraUrl, $"/rest/api/2/project/{projectKey}");
+				Uri uri = new Uri(_jiraBaseUrl, $"/rest/api/2/project/{projectKey}");
 
 				HttpResponseMessage response;
 
@@ -264,7 +328,7 @@ namespace HordeServer.Issues.External
 				// We do get components from the projects endpoint, though these don't contain the archived property :/
 				try
 				{
-					response = await _retryPolicy.ExecuteAsync(ctx => _client.GetAsync(new Uri(_jiraUrl, $"/rest/api/2/project/{projectKey}/components"), ctx), cancellationToken);
+					response = await _retryPolicy.ExecuteAsync(ctx => _client.GetAsync(new Uri(_jiraBaseUrl, $"/rest/api/2/project/{projectKey}/components"), ctx), cancellationToken);
 					response.EnsureSuccessStatusCode();
 				}
 				catch (Exception)
@@ -301,7 +365,7 @@ namespace HordeServer.Issues.External
 		}
 
 		/// <summary>
-		/// Ticks service
+		/// Periodically syncs Horde issue ownership with Jira assignees by matching email addresses
 		/// </summary>
 		async ValueTask TickUpdateAsync(CancellationToken stoppingToken)
 		{
@@ -353,9 +417,11 @@ namespace HordeServer.Issues.External
 			}
 		}
 
-		async ValueTask TickAsync(CancellationToken cancellationToken)
+		/// <summary>
+		/// Periodically refreshes Jira project metadata and issue status for all open Horde issues with linked Jira keys
+		/// </summary>
+		async ValueTask<TimeSpan?> TickAsync(CancellationToken cancellationToken)
 		{
-
 			HashSet<string> jiraProjectKeys = new HashSet<string>();
 
 			BuildConfig buildConfig = _buildConfig.CurrentValue;
@@ -390,6 +456,8 @@ namespace HordeServer.Issues.External
 			}
 			
 			await GetIssuesAsync(jiraKeys.ToArray(), cancellationToken);
+
+			return _tickInterval;
 		}
 
 		public async Task<(string? key, string? url)> CreateIssueAsync(IUser user, string? externalIssueUser, int issueId, string summary, string projectId, string componentId, string issueType, string? description, string? hordeIssueLink, CancellationToken cancellationToken)
@@ -412,7 +480,7 @@ namespace HordeServer.Issues.External
 				issueDesc = $"\n{issueDesc}\n\nJira Issue created by {user.Name} from Horde\n\n";
 			}
 
-			Uri uri = new Uri(_jiraUrl, $"/rest/api/2/issue");
+			Uri uri = new Uri(_jiraBaseUrl, $"/rest/api/2/issue");
 
 			string bodyJson = JsonSerializer.Serialize(new
 			{
@@ -453,7 +521,7 @@ namespace HordeServer.Issues.External
 
 			string responseBody = await response.Content.ReadAsStringAsync(cancellationToken);
 
-			JiraCreateResponse? jiraResponse = JsonSerializer.Deserialize<JiraCreateResponse>(responseBody, new JsonSerializerOptions { PropertyNameCaseInsensitive = true });
+			JiraCreateResponse? jiraResponse = JsonSerializer.Deserialize<JiraCreateResponse>(responseBody, s_jsonOptions);
 
 			if (jiraResponse == null || String.IsNullOrEmpty(jiraResponse.Key))
 			{
@@ -465,7 +533,7 @@ namespace HordeServer.Issues.External
 			// add the user as a watcher to the newly created issue
 			if (externalIssueUser != null)
 			{
-				uri = new Uri(_jiraUrl, $"/rest/api/2/issue/{jiraResponse.Key}/watchers");
+				uri = new Uri(_jiraBaseUrl, $"/rest/api/2/issue/{jiraResponse.Key}/watchers");
 
 				try
 				{
@@ -482,7 +550,7 @@ namespace HordeServer.Issues.External
 				}
 			}
 
-			return (jiraResponse.Key, $"{_jiraUrl}browse/{jiraResponse.Key}");
+			return (jiraResponse.Key, $"{_jiraBaseUrl}/browse/{jiraResponse.Key}");
 		}
 
 		/// <inheritdoc/>
@@ -514,17 +582,36 @@ namespace HordeServer.Issues.External
 
 			if (queryJiras.Count > 0)
 			{
-				Uri uri = new Uri(_jiraUrl, $"/rest/api/2/search?jql=issueKey%20in%20({String.Join(",", queryJiras)})&fields=assignee,status,resolution,priority&maxResults={queryJiras.Count}");
+				HttpResponseMessage response;
+				Uri uri;
+				if (IsCloud)
+				{
+					// Jira Cloud uses the newer /rest/api/3/search/jql endpoint
+					uri = new Uri(_jiraBaseUrl, "/rest/api/3/search/jql");
+					object searchRequest = new
+					{
+						jql = $"issueKey in ({String.Join(",", queryJiras)})",
+						fields = new[] { "assignee", "status", "resolution", "priority" },
+						maxResults = queryJiras.Count
+					};
+					using StringContent content = new StringContent(JsonSerializer.Serialize(searchRequest), Encoding.UTF8, "application/json");
+					response = await _client.PostAsync(uri, content, cancellationToken);
+				}
+				else
+				{
+					// Jira Server/Data Center still supports GET on /rest/api/2/search
+					uri = new Uri(_jiraBaseUrl, $"/rest/api/2/search?jql=issueKey%20in%20({String.Join(",", queryJiras)})&fields=assignee,status,resolution,priority&maxResults={queryJiras.Count}");
+					response = await _client.GetAsync(uri, cancellationToken);
+				}
 
-				HttpResponseMessage response = await _client.GetAsync(uri, cancellationToken);
 				if (!response.IsSuccessStatusCode)
 				{
-					_logger.LogWarning("GET to {Uri} returned {Code} ({Response})", uri, response.StatusCode, await response.Content.ReadAsStringAsync(cancellationToken));
+					_logger.LogWarning("{Method} to {Uri} returned {Code} ({Response})", IsCloud ? "POST" : "GET", uri, response.StatusCode, await response.Content.ReadAsStringAsync(cancellationToken));
 					return result;
 				}
 
 				byte[] data = await response.Content.ReadAsByteArrayAsync(cancellationToken);
-				IssueQueryResponse? jiras = JsonSerializer.Deserialize<IssueQueryResponse>(data.AsSpan(), new JsonSerializerOptions() { PropertyNameCaseInsensitive = true });
+				IssueQueryResponse? jiras = JsonSerializer.Deserialize<IssueQueryResponse>(data.AsSpan(), s_jsonOptions);
 
 				if (jiras != null)
 				{
@@ -534,8 +621,9 @@ namespace HordeServer.Issues.External
 
 						JiraIssue jiraIssue = new JiraIssue();
 						jiraIssue.Key = issue.Key;
-						jiraIssue.Link = $"{_jiraUrl}browse/{issue.Key}";
-						jiraIssue.AssigneeName = issue.Fields?.Assignee?.Name;
+						jiraIssue.Link = $"{_jiraBaseUrl}/browse/{issue.Key}";
+						// v2 (Server/DC) uses Name, v3 (Cloud) falls back to DisplayName then EmailAddress
+						jiraIssue.AssigneeName = issue.Fields?.Assignee?.Name ?? issue.Fields?.Assignee?.DisplayName ?? issue.Fields?.Assignee?.EmailAddress;
 						jiraIssue.AssigneeDisplayName = issue.Fields?.Assignee?.DisplayName;
 						jiraIssue.AssigneeEmailAddress = issue.Fields?.Assignee?.EmailAddress;
 						jiraIssue.PriorityName = issue.Fields?.Priority?.Name;
@@ -595,8 +683,16 @@ namespace HordeServer.Issues.External
 
 		class AssigneeField
 		{
+			/// <summary>
+			/// User name (Server/Data Center only, not present in Cloud v3)
+			/// </summary>
 			public string? Name { get; set; }
 
+			/// <summary>
+			/// Account ID (Cloud only, not present in Server/Data Center)
+			/// </summary>
+			public string? AccountId { get; set; }
+
 			public string? DisplayName { get; set; }
 
 			public string? EmailAddress { get; set; }
@@ -616,5 +712,38 @@ namespace HordeServer.Issues.External
 		{
 			public string? Name { get; set; }
 		}
+
+		/// <summary>
+		/// Jira deployment type
+		/// </summary>
+		enum JiraDeploymentType
+		{
+			/// <summary>
+			/// Unknown deployment type (not yet detected)
+			/// </summary>
+			Unknown,
+
+			/// <summary>
+			/// Jira Cloud deployment
+			/// </summary>
+			Cloud,
+
+			/// <summary>
+			/// Jira Server or Data Center (on-premise) deployment
+			/// </summary>
+			OnPremise
+		}
+
+		/// <summary>
+		/// Response from /rest/api/2/serverInfo endpoint
+		/// </summary>
+		class ServerInfoResponse
+		{
+			public string? BaseUrl { get; set; }
+			public string? Version { get; set; }
+			public string? DeploymentType { get; set; }
+			public int[]? VersionNumbers { get; set; }
+			public string? ServerTitle { get; set; }
+		}
 	}
 }
