Horde: Add lease divisions in agent telemetry view

diff --git a/Engine/Source/Programs/Horde/HordeDashboard/src/components/agents/HistoryModal.tsx b/Engine/Source/Programs/Horde/HordeDashboard/src/components/agents/HistoryModal.tsx
--- a/Engine/Source/Programs/Horde/HordeDashboard/src/components/agents/HistoryModal.tsx
+++ b/Engine/Source/Programs/Horde/HordeDashboard/src/components/agents/HistoryModal.tsx
@@ -13,7 +13,7 @@ import { getNiceTime, getShortNiceTime } from "../../base/utilities/timeUtils";
 import { getHordeStyling } from "../../styles/Styles";
 import { getHordeTheme } from "../../styles/theme";
 import { BatchStatusIcon, LeaseStatusIcon, StepStatusIcon } from "./../StatusIcon";
-import { historyModalStore } from "./HistoryModalStore";
+import { historyModalStore, processLeaseData, ProcessedLease } from "./HistoryModalStore";
 import { formatBytes } from "horde/base/utilities/stringUtills";
 
 const historyStyles = mergeStyleSets({
@@ -59,7 +59,7 @@ const TelemetryTooltip: React.FC<{ id: "cpu" | "ram" | "disk", cpuKey?: string }
    const translateY = "-50%";
 
    const textSize = 11;
-   const titleWidth = 48;
+   const titleWidth = 58;
 
    const dataElement = (title: string, value: string, link?: string, agent?: string, linkTarget?: string) => {
 
@@ -115,13 +115,39 @@ const TelemetryTooltip: React.FC<{ id: "cpu" | "ram" | "disk", cpuKey?: string }
    }
 
    if (lease?.details?.logId) {
-      elements.push(dataElement("Log:", `Lease`, `/log/${lease.details.logId}`));
+      elements.push(dataElement("Log:", `Lease (${lease.id})`, `/log/${lease.details.logId}`));
    }
 
    if (tooltip.time) {
       elements.push(dataElement("Time:", getNiceTime(tooltip.time)));
    }
 
+   // Calculate and display lease duration
+   if (lease) {
+      const startMs = new Date(lease.startTime).getTime();
+      const endMs = lease.finishTime ? new Date(lease.finishTime).getTime() : Date.now();
+      const durationMs = endMs - startMs;
+
+      const totalSeconds = Math.floor(durationMs / 1000);
+      const days = Math.floor(totalSeconds / 86400);
+      const hours = Math.floor((totalSeconds % 86400) / 3600);
+      const minutes = Math.floor((totalSeconds % 3600) / 60);
+      const seconds = totalSeconds % 60;
+
+      const parts: string[] = [];
+      if (days > 0) parts.push(`${days}d`);
+      if (hours > 0 || days > 0) parts.push(`${hours}h`);
+      if (minutes > 0 || hours > 0 || days > 0) parts.push(`${minutes}m`);
+      parts.push(`${seconds}s`);
+
+      let durationStr = parts.join(" ");
+      if (!lease.finishTime) {
+         durationStr += " (ongoing)";
+      }
+
+      elements.push(dataElement("Duration:", durationStr));
+   }
+
    if (tooltip.sample) {
 
       const d = tooltip.sample;
@@ -223,6 +249,125 @@ const cpuSelections: IComboBoxOption[] = [
    }
 ]
 
+// 10 distinct colors for lease overlays - more transparent to not occlude graph data
+const leaseColorsDark = [
+   "rgba(100, 149, 237, 0.15)",  // Cornflower blue
+   "rgba(144, 238, 144, 0.15)",  // Light green
+   "rgba(255, 182, 108, 0.15)",  // Peach
+   "rgba(221, 160, 221, 0.15)",  // Plum
+   "rgba(64, 224, 208, 0.15)",   // Turquoise
+   "rgba(255, 218, 185, 0.15)",  // Peach puff
+   "rgba(135, 206, 235, 0.15)",  // Sky blue
+   "rgba(255, 160, 122, 0.15)",  // Light salmon
+   "rgba(152, 251, 152, 0.15)",  // Pale green
+   "rgba(216, 191, 216, 0.15)",  // Thistle
+];
+
+const leaseColorsLight = [
+   "rgba(65, 105, 225, 0.12)",   // Royal blue
+   "rgba(34, 139, 34, 0.12)",    // Forest green
+   "rgba(255, 140, 0, 0.12)",    // Dark orange
+   "rgba(148, 0, 211, 0.12)",    // Dark violet
+   "rgba(0, 139, 139, 0.12)",    // Dark cyan
+   "rgba(184, 134, 11, 0.12)",   // Dark goldenrod
+   "rgba(70, 130, 180, 0.12)",   // Steel blue
+   "rgba(205, 92, 92, 0.12)",    // Indian red
+   "rgba(60, 179, 113, 0.12)",   // Medium sea green
+   "rgba(186, 85, 211, 0.12)",   // Medium orchid
+];
+
+// Faint vertical line colors for lease boundaries
+const leaseBorderColorDark = "rgba(200, 200, 200, 0.3)";
+const leaseBorderColorLight = "rgba(80, 80, 80, 0.25)";
+
+const LeaseOverlayRects: React.FC<{ processedLeases: ProcessedLease[], sparkHeight: number, sparkWidth: number, chartId: string }> = ({ processedLeases, sparkHeight, sparkWidth, chartId }) => {
+
+   const colors = dashboard.darktheme ? leaseColorsDark : leaseColorsLight;
+   const borderColor = dashboard.darktheme ? leaseBorderColorDark : leaseBorderColorLight;
+
+   if (!processedLeases.length) {
+      return null;
+   }
+
+   const elements: JSX.Element[] = [];
+
+   // Sort leases by start position to properly detect adjacency
+   const sortedLeases = [...processedLeases].sort((a, b) => a.startX - b.startX);
+
+   sortedLeases.forEach((pl, idx) => {
+      const color = colors[processedLeases.indexOf(pl) % colors.length];
+      let x = pl.startX;
+      let endX = pl.endX !== null ? pl.endX : sparkWidth;
+
+      // Extend to cover gaps: check if previous lease ends near this one's start
+      if (idx > 0) {
+         const prevLease = sortedLeases[idx - 1];
+         const prevEndX = prevLease.endX !== null ? prevLease.endX : sparkWidth;
+         // If there's a small gap (< 2px), extend this rect to fill it
+         if (x - prevEndX > 0 && x - prevEndX < 2) {
+            x = prevEndX;
+         }
+      }
+
+      // Extend end to next lease if there's a small gap
+      if (idx < sortedLeases.length - 1) {
+         const nextLease = sortedLeases[idx + 1];
+         // If there's a small gap (< 2px), extend this rect to fill it
+         if (nextLease.startX - endX > 0 && nextLease.startX - endX < 2) {
+            endX = nextLease.startX;
+         }
+      }
+
+      const width = endX - x;
+
+      if (width > 0) {
+         // Rectangle fill
+         elements.push(
+            <rect
+               key={`${chartId}-lease-${idx}`}
+               x={x}
+               y={0}
+               width={width}
+               height={sparkHeight}
+               fill={color}
+            />
+         );
+
+         // Start line (only if not at the very beginning)
+         if (pl.startPercent > 0.01) {
+            elements.push(
+               <line
+                  key={`${chartId}-start-${idx}`}
+                  x1={pl.startX}
+                  x2={pl.startX}
+                  y1={0}
+                  y2={sparkHeight}
+                  stroke={borderColor}
+                  strokeWidth={1}
+               />
+            );
+         }
+
+         // End line (only if lease has finished and not at the very end)
+         if (pl.endX !== null && pl.endPercent !== null && pl.endPercent < 0.99) {
+            elements.push(
+               <line
+                  key={`${chartId}-end-${idx}`}
+                  x1={pl.endX}
+                  x2={pl.endX}
+                  y1={0}
+                  y2={sparkHeight}
+                  stroke={borderColor}
+                  strokeWidth={1}
+               />
+            );
+         }
+      }
+   });
+
+   return <>{elements}</>;
+};
+
 const TelemetryPanel: React.FC<{}> = observer(({ }) => {
 
    const [timeKey, setTimeKey] = useState("time_2_days");
@@ -255,6 +400,7 @@ const TelemetryPanel: React.FC<{}> = observer(({ }) => {
    const statusColors = dashboard.getStatusColors();
 
    const sparkWidth = 1100;
+   const sparkHeight = 128;
    const data = (currentData[0] as GetAgentTelemetrySampleResponse[]);
    const leases = ((currentData.length > 1) ? currentData[1] : []) as GetAgentLeaseResponse[];
 
@@ -262,6 +408,9 @@ const TelemetryPanel: React.FC<{}> = observer(({ }) => {
       return null;
    }
 
+   const processedData = processLeaseData(data, leases, sparkWidth);
+   const processedLeases = processedData?.processedLeases ?? [];
+
    const cpuData = data.map(d => {
       if (cpuKey == "user") return d.userCpu;
       if (cpuKey == "system") return d.systemCpu;
@@ -332,6 +481,9 @@ const TelemetryPanel: React.FC<{}> = observer(({ }) => {
    })
 
 
+   // Padding value used in Sparklines style - must match for coordinate calculations
+   const sparkPadding = 8;
+
    function relativeCoords(event: any, id: string) {
 
       if (state.currentLeaseTooltip?.frozen || !data?.length) {
@@ -342,8 +494,14 @@ const TelemetryPanel: React.FC<{}> = observer(({ }) => {
       let x = event.clientX - bounds.left;
       let y = event.clientY - bounds.top;
 
-      const percent = ((x) / 1130);
-      const sampleIdx = Math.ceil(data.length * percent);
+      // Adjust for padding and border to get position within the actual chart area
+      const chartX = x - sparkPadding - 1; // -1 for border
+      const percent = chartX / sparkWidth;
+
+      // Clamp and calculate sample index
+      const clampedPercent = Math.max(0, Math.min(1, percent));
+      const sampleIdx = Math.min(Math.floor(clampedPercent * data.length), data.length - 1);
+
       if (sampleIdx < 0 || sampleIdx >= data.length) {
          return;
       }
@@ -352,8 +510,11 @@ const TelemetryPanel: React.FC<{}> = observer(({ }) => {
       const time = sample.time;
 
       let lease: GetAgentLeaseResponse | undefined;
+      const timeMs = new Date(time).getTime();
       leases.forEach(s => {
-         if (time >= s.startTime && (!s.finishTime || time <= s.finishTime)) {
+         const startMs = new Date(s.startTime).getTime();
+         const finishMs = s.finishTime ? new Date(s.finishTime).getTime() : null;
+         if (timeMs >= startMs && (finishMs === null || timeMs <= finishMs)) {
             lease = s;
          }
       });
@@ -422,13 +583,13 @@ const TelemetryPanel: React.FC<{}> = observer(({ }) => {
                <TelemetryTooltip id={"cpu"} cpuKey={cpuKey} />
             </div>
 
-            <Sparklines width={sparkWidth} height={128} data={cpuData} min={0} max={100} style={{ backgroundColor: dashboard.darktheme ? "#060709" : "#F3F2F1", padding: 8, border: "solid 1px #181A1B" }}>
+            <Sparklines width={sparkWidth} height={sparkHeight} data={cpuData} min={0} max={100} style={{ backgroundColor: dashboard.darktheme ? "#060709" : "#F3F2F1", padding: 8, border: "solid 1px #181A1B" }}>
                <SparklinesLine color={dashboard.darktheme ? "lightblue" : "#1E90FF"} />
                <SparklinesReferenceLine type="custom" value={0} style={{ stroke: statusColors.get(StatusColor.Failure) }} />
                <SparklinesReferenceLine type="custom" value={32} style={{ stroke: statusColors.get(StatusColor.Warnings) }} />
                <SparklinesReferenceLine type="custom" value={64} style={{ stroke: statusColors.get(StatusColor.Success) }} />
                <SparklinesReferenceLine type="custom" value={96} style={{ stroke: statusColors.get(StatusColor.Success) }} />
-
+               <LeaseOverlayRects processedLeases={processedLeases} sparkHeight={sparkHeight} sparkWidth={sparkWidth} chartId="cpu" />
             </Sparklines>
          </div>
       </Stack>
@@ -447,13 +608,13 @@ const TelemetryPanel: React.FC<{}> = observer(({ }) => {
             <div style={{ position: "relative" }}>
                <TelemetryTooltip id="ram" />
             </div>
-            <Sparklines width={sparkWidth} height={128} data={ramData} min={0} max={1} style={{ backgroundColor: dashboard.darktheme ? "#060709" : "#F3F2F1", padding: 8, border: "solid 1px #181A1B" }}>
+            <Sparklines width={sparkWidth} height={sparkHeight} data={ramData} min={0} max={1} style={{ backgroundColor: dashboard.darktheme ? "#060709" : "#F3F2F1", padding: 8, border: "solid 1px #181A1B" }}>
                <SparklinesLine color={dashboard.darktheme ? "lightblue" : "#1E90FF"} />
                <SparklinesReferenceLine type="custom" value={0} style={{ stroke: statusColors.get(StatusColor.Failure) }} />
                <SparklinesReferenceLine type="custom" value={32} style={{ stroke: statusColors.get(StatusColor.Warnings) }} />
                <SparklinesReferenceLine type="custom" value={64} style={{ stroke: statusColors.get(StatusColor.Success) }} />
                <SparklinesReferenceLine type="custom" value={96} style={{ stroke: statusColors.get(StatusColor.Success) }} />
-
+               <LeaseOverlayRects processedLeases={processedLeases} sparkHeight={sparkHeight} sparkWidth={sparkWidth} chartId="ram" />
             </Sparklines>
          </div>
       </Stack>
@@ -473,12 +634,13 @@ const TelemetryPanel: React.FC<{}> = observer(({ }) => {
             <div style={{ position: "relative" }}>
                <TelemetryTooltip id="disk" />
             </div>
-            <Sparklines width={sparkWidth} height={128} data={diskData} min={0} max={1} style={{ backgroundColor: dashboard.darktheme ? "#060709" : "#F3F2F1", padding: 8, border: "solid 1px #181A1B" }}>
+            <Sparklines width={sparkWidth} height={sparkHeight} data={diskData} min={0} max={1} style={{ backgroundColor: dashboard.darktheme ? "#060709" : "#F3F2F1", padding: 8, border: "solid 1px #181A1B" }}>
                <SparklinesLine color={dashboard.darktheme ? "lightblue" : "#1E90FF"} />
                <SparklinesReferenceLine type="custom" value={0} style={{ stroke: statusColors.get(StatusColor.Failure) }} />
                <SparklinesReferenceLine type="custom" value={32} style={{ stroke: statusColors.get(StatusColor.Warnings) }} />
                <SparklinesReferenceLine type="custom" value={64} style={{ stroke: statusColors.get(StatusColor.Success) }} />
                <SparklinesReferenceLine type="custom" value={96} style={{ stroke: statusColors.get(StatusColor.Success) }} />
+               <LeaseOverlayRects processedLeases={processedLeases} sparkHeight={sparkHeight} sparkWidth={sparkWidth} chartId="disk" />
             </Sparklines>
          </div>
       </Stack>}
diff --git a/Engine/Source/Programs/Horde/HordeDashboard/src/components/agents/HistoryModalStore.ts b/Engine/Source/Programs/Horde/HordeDashboard/src/components/agents/HistoryModalStore.ts
--- a/Engine/Source/Programs/Horde/HordeDashboard/src/components/agents/HistoryModalStore.ts
+++ b/Engine/Source/Programs/Horde/HordeDashboard/src/components/agents/HistoryModalStore.ts
@@ -14,6 +14,22 @@ type LeaseTooltip = {
    sample?: GetAgentTelemetrySampleResponse;
 }
 
+export type ProcessedLease = {
+   lease: GetAgentLeaseResponse;
+   startX: number;
+   endX: number | null;
+   startPercent: number;
+   endPercent: number | null;
+}
+
+export type ProcessedTelemetryData = {
+   samples: GetAgentTelemetrySampleResponse[];
+   leases: GetAgentLeaseResponse[];
+   processedLeases: ProcessedLease[];
+   minTime: Date;
+   maxTime: Date;
+}
+
 type InfoPanelItem = {
    key: string;
    name: string;
@@ -26,6 +42,96 @@ type InfoPanelSubItem = {
    value: string;
 }
 
+export function processLeaseData(
+   samples: GetAgentTelemetrySampleResponse[],
+   leases: GetAgentLeaseResponse[],
+   sparkWidth: number
+): ProcessedTelemetryData | null {
+   if (!samples.length) {
+      return null;
+   }
+
+   const minTime = samples[0].time;
+   const maxTime = samples[samples.length - 1].time;
+   const minTimeMs = new Date(minTime).getTime();
+   const maxTimeMs = new Date(maxTime).getTime();
+   const timeRangeMs = maxTimeMs - minTimeMs;
+
+   if (timeRangeMs <= 0) {
+      return null;
+   }
+
+   const processedLeases: ProcessedLease[] = [];
+
+   // Sparklines distributes points evenly by index, so we need to find which
+   // sample index corresponds to the lease time and interpolate the position
+   const numSamples = samples.length;
+
+   // Helper to find interpolated position for a given time
+   const findInterpolatedPosition = (targetMs: number): number => {
+      // Clamp to visible range
+      if (targetMs <= minTimeMs) return 0;
+      if (targetMs >= maxTimeMs) return 1;
+
+      // Find the two samples that bracket the target time
+      for (let i = 0; i < numSamples - 1; i++) {
+         const currentMs = new Date(samples[i].time).getTime();
+         const nextMs = new Date(samples[i + 1].time).getTime();
+
+         if (targetMs >= currentMs && targetMs <= nextMs) {
+            // Interpolate between sample i and i+1
+            const timeFraction = (targetMs - currentMs) / (nextMs - currentMs);
+            const idxFraction = (i + timeFraction) / (numSamples - 1);
+            return idxFraction;
+         }
+      }
+
+      // Fallback: use time-based percentage
+      return (targetMs - minTimeMs) / timeRangeMs;
+   };
+
+   for (const lease of leases) {
+      const leaseStartMs = new Date(lease.startTime).getTime();
+      const leaseEndMs = lease.finishTime ? new Date(lease.finishTime).getTime() : null;
+
+      // Skip leases entirely outside the visible range
+      if (leaseEndMs !== null && leaseEndMs < minTimeMs) {
+         continue;
+      }
+      if (leaseStartMs > maxTimeMs) {
+         continue;
+      }
+
+      // Calculate start position with interpolation
+      const startPercent = findInterpolatedPosition(leaseStartMs);
+      const startX = startPercent * sparkWidth;
+
+      // Calculate end position with interpolation
+      let endPercent: number | null = null;
+      let endX: number | null = null;
+      if (leaseEndMs !== null) {
+         endPercent = findInterpolatedPosition(leaseEndMs);
+         endX = endPercent * sparkWidth;
+      }
+
+      processedLeases.push({
+         lease,
+         startX,
+         endX,
+         startPercent,
+         endPercent
+      });
+   }
+
+   return {
+      samples,
+      leases,
+      processedLeases,
+      minTime,
+      maxTime
+   };
+}
+
 class HistoryModalStore {
    constructor() {
       makeObservable(this);
