Horde: Address issues with on prem Jira service build health

diff --git a/Engine/Source/Programs/Horde/Plugins/Build/HordeServer.Build/Issues/External/JiraService.cs b/Engine/Source/Programs/Horde/Plugins/Build/HordeServer.Build/Issues/External/JiraService.cs
--- a/Engine/Source/Programs/Horde/Plugins/Build/HordeServer.Build/Issues/External/JiraService.cs
+++ b/Engine/Source/Programs/Horde/Plugins/Build/HordeServer.Build/Issues/External/JiraService.cs
@@ -97,9 +97,13 @@ namespace HordeServer.Issues.External
 		static readonly JsonSerializerOptions s_jsonOptions = new JsonSerializerOptions { PropertyNameCaseInsensitive = true };
 
 		/// <summary>
-		/// Normalizes a URL by removing any trailing slashes
+		/// Normalizes a URL by ensuring it has a trailing slash
 		/// </summary>
-		static Uri NormalizeUrl(Uri url) => new Uri(url.ToString().TrimEnd('/'));
+		static Uri NormalizeUrl(Uri url)
+		{
+			string urlString = url.ToString();
+			return urlString.EndsWith('/') ? url : new Uri(urlString + '/');
+		}
 
 		readonly ILogger _logger;
 
@@ -192,7 +196,7 @@ namespace HordeServer.Issues.External
 		{
 			try
 			{
-				Uri uri = new Uri(_jiraUrl, "/rest/api/2/serverInfo");
+				Uri uri = new Uri(_jiraUrl, "rest/api/2/serverInfo");
 				HttpResponseMessage response = await _retryPolicy.ExecuteAsync(ctx => _client.GetAsync(uri, ctx), cancellationToken);
 				response.EnsureSuccessStatusCode();
 
@@ -232,7 +236,7 @@ namespace HordeServer.Issues.External
 		}
 
 		/// <inheritdoc/>
-		public string? GetIssueUrl(string key) => $"{_jiraBaseUrl}/browse/{key}";
+		public string? GetIssueUrl(string key) => $"{_jiraBaseUrl}browse/{key}";
 
 		public Task<List<IExternalIssueProject>> GetProjectsAsync(StreamConfig streamConfig, CancellationToken cancellationToken)
 		{
@@ -274,7 +278,7 @@ namespace HordeServer.Issues.External
 			{
 				string projectKey = jiraProjectKeys[i];
 
-				Uri uri = new Uri(_jiraBaseUrl, $"/rest/api/2/project/{projectKey}");
+				Uri uri = new Uri(_jiraBaseUrl, $"rest/api/2/project/{projectKey}");
 
 				HttpResponseMessage response;
 
@@ -328,7 +332,7 @@ namespace HordeServer.Issues.External
 				// We do get components from the projects endpoint, though these don't contain the archived property :/
 				try
 				{
-					response = await _retryPolicy.ExecuteAsync(ctx => _client.GetAsync(new Uri(_jiraBaseUrl, $"/rest/api/2/project/{projectKey}/components"), ctx), cancellationToken);
+					response = await _retryPolicy.ExecuteAsync(ctx => _client.GetAsync(new Uri(_jiraBaseUrl, $"rest/api/2/project/{projectKey}/components"), ctx), cancellationToken);
 					response.EnsureSuccessStatusCode();
 				}
 				catch (Exception)
@@ -480,7 +484,7 @@ namespace HordeServer.Issues.External
 				issueDesc = $"\n{issueDesc}\n\nJira Issue created by {user.Name} from Horde\n\n";
 			}
 
-			Uri uri = new Uri(_jiraBaseUrl, $"/rest/api/2/issue");
+			Uri uri = new Uri(_jiraBaseUrl, "rest/api/2/issue");
 
 			string bodyJson = JsonSerializer.Serialize(new
 			{
@@ -533,7 +537,7 @@ namespace HordeServer.Issues.External
 			// add the user as a watcher to the newly created issue
 			if (externalIssueUser != null)
 			{
-				uri = new Uri(_jiraBaseUrl, $"/rest/api/2/issue/{jiraResponse.Key}/watchers");
+				uri = new Uri(_jiraBaseUrl, $"rest/api/2/issue/{jiraResponse.Key}/watchers");
 
 				try
 				{
@@ -550,7 +554,7 @@ namespace HordeServer.Issues.External
 				}
 			}
 
-			return (jiraResponse.Key, $"{_jiraBaseUrl}/browse/{jiraResponse.Key}");
+			return (jiraResponse.Key, $"{_jiraBaseUrl}browse/{jiraResponse.Key}");
 		}
 
 		/// <inheritdoc/>
@@ -587,7 +591,7 @@ namespace HordeServer.Issues.External
 				if (IsCloud)
 				{
 					// Jira Cloud uses the newer /rest/api/3/search/jql endpoint
-					uri = new Uri(_jiraBaseUrl, "/rest/api/3/search/jql");
+					uri = new Uri(_jiraBaseUrl, "rest/api/3/search/jql");
 					object searchRequest = new
 					{
 						jql = $"issueKey in ({String.Join(",", queryJiras)})",
@@ -600,7 +604,7 @@ namespace HordeServer.Issues.External
 				else
 				{
 					// Jira Server/Data Center still supports GET on /rest/api/2/search
-					uri = new Uri(_jiraBaseUrl, $"/rest/api/2/search?jql=issueKey%20in%20({String.Join(",", queryJiras)})&fields=assignee,status,resolution,priority&maxResults={queryJiras.Count}");
+					uri = new Uri(_jiraBaseUrl, $"rest/api/2/search?jql=issueKey%20in%20({String.Join(",", queryJiras)})&fields=assignee,status,resolution,priority&maxResults={queryJiras.Count}");
 					response = await _client.GetAsync(uri, cancellationToken);
 				}
 
@@ -621,7 +625,7 @@ namespace HordeServer.Issues.External
 
 						JiraIssue jiraIssue = new JiraIssue();
 						jiraIssue.Key = issue.Key;
-						jiraIssue.Link = $"{_jiraBaseUrl}/browse/{issue.Key}";
+						jiraIssue.Link = $"{_jiraBaseUrl}browse/{issue.Key}";
 						// v2 (Server/DC) uses Name, v3 (Cloud) falls back to DisplayName then EmailAddress
 						jiraIssue.AssigneeName = issue.Fields?.Assignee?.Name ?? issue.Fields?.Assignee?.DisplayName ?? issue.Fields?.Assignee?.EmailAddress;
 						jiraIssue.AssigneeDisplayName = issue.Fields?.Assignee?.DisplayName;
