[Horde][Deploy][Dashboard] Release notes

#Synced-CL 49813167

diff --git a/Engine/Source/Programs/Horde/Docs/ReleaseNotes.md b/Engine/Source/Programs/Horde/Docs/ReleaseNotes.md
--- a/Engine/Source/Programs/Horde/Docs/ReleaseNotes.md
+++ b/Engine/Source/Programs/Horde/Docs/ReleaseNotes.md
@@ -1,5 +1,42 @@
 # Release Notes
 
+## 2026-01-15
+
+* Address issues with on prem Jira service build health (49818828)
+* Add lease divisions in agent telemetry view (49814164)
+* Use MSBuildThisFileDirectory for MetaData.cs path, this allows for dotnet run <project> to work as it uses a different working directory (49813633)
+* Support Jira Cloud and On Premise for Build Health (49809016)
+* Fix conform in PerforceMaterializer to run cleanup after sync and improve error handling (49767276)
+* Add GetJobsAsync to HordeHttpClient (49700421)
+* Expose optional force parameter for the device reservation deletion endpoint
+  This will allow users to fully delete reservations and associated reservation blocks where appropriate (49690413)
+* Fix Telemetry Service to use linked CancellationTokenSource for FLTMC command (49647092)
+* Conforms that are scheduled due to low free disk space upon execution check if the agent's available disk space still requires the conform to be run (49624335)
+* The task argument for jobs and conforms is passed using an environment variable (49619111)
+* Create jobid,batchid,stepid column hierarchy within JobStepRef. (49604444)
+* Adding metadata & specific summary table emissions (49602718)
+* Cleaning up telemetry records to use base class for common attributes & supporting inheritance for AnalyticsTableGenerator. (49595510)
+* Fix potential issue with CPU time calculations in agent
+  Also wrap lease execution to ensure any exceptions are immediately surfaced. (49414139)
+* Cloud Storage - Added another longer wait (this time for 60 minutes before retrying). We have seen extreme cases were issues take about an hour to settle and a retry in this scenario would have recovered the build rather then forcing us to rerun the build which takes much longer then this 1 hour wait time. (49372727)
+* Instrument gate handling for schedules (49372088)
+* Refreshed stale files from shelve for build break (49354227)
+* Build health filter QoL improvements, preparation for wide users, and minor bug fixes:
+  - Description field
+  - Full ownership context returned with object
+  - Dropdown to use seach & items ordered alphabetically
+  - Dropdown tooltips to include owner, proper time format, and description
+  - Labels more intuitive based on whether the owner has authority to CRUD build filters, or not
+  - Syntheized job step outcomes for when batches fail (49353237)
+* Description field & admin powers for Build health filter, as well as returning full owner object for better messaging around follow-ups. (49353218)
+* Fix issue with agent tooltip on job timeline (49332510)
+* Another take on duplicate json Telemetry key exceptions (49282234)
+* Move ISystemMetrics to shared (49278929)
+* Fix an analysis warning with EpicGames.Analytics.Generators. (49272410)
+* Add a runtime watchdog to agent to monitor for GC or thread pool stalls (49209914)
+* Fix exception during PerforceMaterializer clean when in dirty state
+  Parallel directory deletion did assume workspace directory exists, which is not always true. (49204635)
+
 ## 2025-12-11
 
 * Using nullable type for IJob as per CreateJobObject return specification. (49147810)
