Horde: Adds IPluginDashboardAcls interface for plugins to control dashboard feature visibility based on user permissions

diff --git a/Engine/Source/Programs/Horde/HordeDashboard/plugins/index.ts b/Engine/Source/Programs/Horde/HordeDashboard/plugins/index.ts
--- a/Engine/Source/Programs/Horde/HordeDashboard/plugins/index.ts
+++ b/Engine/Source/Programs/Horde/HordeDashboard/plugins/index.ts
@@ -1,6 +1,7 @@
 // Copyright Epic Games, Inc. All Rights Reserved.
 
 import { RouteObject } from "react-router-dom";
+import { ServerPluginInfoResponse } from "horde/backend/Api";
 
 // Routes => components
 // backend queries
@@ -15,6 +16,8 @@ export type Mount = {
    context?: string;
    text: string;
    route: string;
+   // ACLs required to show this mount point (all must be present)
+   requiredAcls?: string[];
 }
 
 export type HordePlugin = {
@@ -25,6 +28,10 @@ export type HordePlugin = {
    // multiple plugin mount points
    mounts?: Mount[];
    enabled?: boolean;
+   // ACLs required to enable the entire plugin (all must be present)
+   requiredAcls?: string[];
+   // User's ACLs for this plugin (populated from server response)
+   userAcls?: string[];
 }
 
 const plugins: HordePlugin[] = [];
@@ -38,10 +45,41 @@ export function registerHordePlugin(plugin: HordePlugin) {
    plugins.push(plugin);
 }
 
-export function enableHordePlugins(ids: string[]) {
-   const set = new Set(ids);
+/**
+ * Enables plugins based on server plugin info, checking both loaded status and ACL requirements.
+ *
+ * ACL semantics from server:
+ * - undefined/null: Plugin doesn't implement dashboard ACLs (no restrictions)
+ * - empty array []: Plugin uses ACLs but user has no permissions
+ * - non-empty array: User's granted ACL actions for this plugin
+ *
+ * @param serverPlugins Plugin information from the server, including user's ACLs
+ */
+export function enableHordePlugins(serverPlugins: ServerPluginInfoResponse[]) {
+   // Use lowercase keys for case-insensitive matching between frontend plugin IDs and server plugin names
+   const pluginMap = new Map(serverPlugins.map(p => [p.name.toLowerCase(), p]));
+
    plugins.forEach(p => {
-      p.enabled = set.has(p.id)
+      const serverPlugin = pluginMap.get(p.id.toLowerCase());
+
+      // Plugin must be loaded on the server
+      if (!serverPlugin?.loaded) {
+         p.enabled = false;
+         p.userAcls = undefined;
+         return;
+      }
+
+      // Store the user's ACLs for this plugin (undefined means no ACL restrictions)
+      p.userAcls = serverPlugin.userAcls;
+
+      // Check plugin-level ACL requirements
+      if (p.requiredAcls?.length) {
+         const userAclSet = new Set(serverPlugin.userAcls ?? []);
+         const hasAllRequired = p.requiredAcls.every(acl => userAclSet.has(acl));
+         p.enabled = hasAllRequired;
+      } else {
+         p.enabled = true;
+      }
    });
 }
 
@@ -49,3 +87,21 @@ export function getHordePlugins(): HordePlugin[] {
    return plugins.filter(p => !!p.enabled);
 }
 
+/**
+ * Gets the filtered mounts for a plugin based on the user's ACLs
+ * @param plugin The plugin to get mounts for
+ * @returns Array of mounts the user has access to
+ */
+export function getFilteredMounts(plugin: HordePlugin): Mount[] {
+   const allMounts = plugin.mounts ?? (plugin.mount ? [plugin.mount] : []);
+   const userAclSet = new Set(plugin.userAcls ?? []);
+
+   return allMounts.filter(mount => {
+      // If no required ACLs, mount is visible to all
+      if (!mount.requiredAcls?.length) {
+         return true;
+      }
+      // All required ACLs must be present
+      return mount.requiredAcls.every(acl => userAclSet.has(acl));
+   });
+}
diff --git a/Engine/Source/Programs/Horde/HordeDashboard/src/App.tsx b/Engine/Source/Programs/Horde/HordeDashboard/src/App.tsx
--- a/Engine/Source/Programs/Horde/HordeDashboard/src/App.tsx
+++ b/Engine/Source/Programs/Horde/HordeDashboard/src/App.tsx
@@ -125,7 +125,7 @@ const Main: React.FC = () => {
       console.log("No Enabled Plugins")
    }
 
-   enableHordePlugins(enabledPlugins.map(p => p.name));
+   enableHordePlugins(enabledPlugins);
    const plugins = getHordePlugins();
 
    plugins.forEach(p => {
diff --git a/Engine/Source/Programs/Horde/HordeDashboard/src/backend/Api.ts b/Engine/Source/Programs/Horde/HordeDashboard/src/backend/Api.ts
--- a/Engine/Source/Programs/Horde/HordeDashboard/src/backend/Api.ts
+++ b/Engine/Source/Programs/Horde/HordeDashboard/src/backend/Api.ts
@@ -4510,6 +4510,8 @@ export type ServerPluginInfoResponse = {
 	description?: string;
 	loaded: boolean,
 	version?: string;
+	/// List of ACL actions the current user has for this plugin
+	userAcls?: string[];
 };
 
 export type GetServerInfoResponse = {
diff --git a/Engine/Source/Programs/Horde/HordeDashboard/src/components/TopNav.tsx b/Engine/Source/Programs/Horde/HordeDashboard/src/components/TopNav.tsx
--- a/Engine/Source/Programs/Horde/HordeDashboard/src/components/TopNav.tsx
+++ b/Engine/Source/Programs/Horde/HordeDashboard/src/components/TopNav.tsx
@@ -13,7 +13,7 @@ import { VersionModal } from './VersionModal';
 import { getHordeTheme } from '../styles/theme';
 import { getHordeStyling } from '../styles/Styles';
 import { ArtifactQueryState, FindArtifactsModal } from './artifacts/ArtifactsSearch';
-import { getHordePlugins, MountType } from 'hordePlugins';
+import { getHordePlugins, getFilteredMounts, MountType } from 'hordePlugins';
 import { SupportModal } from './SupportModal';
 
 const getStyles = () => {
@@ -442,30 +442,19 @@ export const TopNav: React.FC<{ suppressServer?: boolean }> = observer(({ suppre
 
       let serviceItems: IContextualMenuItem[] = [];
 
-      // get the plugins which mount into the tools bar
-      const plugins = getHordePlugins().filter(p => p.mount?.type === MountType.TopNav || p.mounts?.find(m => m.type === MountType.TopNav));
+      // get the plugins and their filtered mounts (based on user ACLs)
+      const plugins = getHordePlugins();
 
-      // Add plugin topnav mount points
+      // Add plugin topnav mount points (filtered by user ACLs)
       plugins.forEach(p => {
-         if (p.mount) {
-            serviceItems.push({
-               key: `admin_plugin_${p.mount!.text}`,
-               text: p.mount!.text,
-               link: p.mount!.route
-            });      
-         }
-
-         p.mounts?.forEach( m => {
-            if (m.type !== MountType.TopNav) {
-               return;
-            }
-
+         const mounts = getFilteredMounts(p).filter(m => m.type === MountType.TopNav);
+         mounts.forEach(m => {
             serviceItems.push({
                key: `admin_plugin_${m.text}`,
                text: m.text,
                link: m.route
-            });      
-         })
+            });
+         });
       })
 
       if (features?.showTests !== false) {
diff --git a/Engine/Source/Programs/Horde/HordeDashboard/src/components/VersionModal.tsx b/Engine/Source/Programs/Horde/HordeDashboard/src/components/VersionModal.tsx
--- a/Engine/Source/Programs/Horde/HordeDashboard/src/components/VersionModal.tsx
+++ b/Engine/Source/Programs/Horde/HordeDashboard/src/components/VersionModal.tsx
@@ -13,6 +13,7 @@ export const VersionModal: React.FC<{ show: boolean, onClose: () => void }> = ({
    type GeneralItem = {
       name: string;
       value?: string;
+      userAcls?: string[];
    }
 
    const close = () => {
@@ -35,8 +36,9 @@ export const VersionModal: React.FC<{ show: boolean, onClose: () => void }> = ({
    const versionItems: GeneralItem[] = [];
 
    const columns = [
-      { key: 'column1', name: 'Name', fieldName: 'name', minWidth: 100, maxWidth: 100 },
+      { key: 'column1', name: 'Name', fieldName: 'name', minWidth: 120, maxWidth: 120 },
       { key: 'column2', name: 'Value', fieldName: 'value', minWidth: 100, maxWidth: 200 },
+      { key: 'column3', name: 'User ACLs', fieldName: 'userAcls', minWidth: 200, maxWidth: 400 },
    ];
 
    let dashboardVersion: string | undefined;
@@ -65,17 +67,16 @@ export const VersionModal: React.FC<{ show: boolean, onClose: () => void }> = ({
       versionItems.push({ name: "Plugins:", value: "" });
 
       version.serverInfo?.plugins.forEach(p => {
-         versionItems.push({ name: p.name, value: p.version ?? "" });
+         versionItems.push({ name: p.name, value: p.version ?? "", userAcls: p.userAcls });
       })
    
    }
 
 
-   const onRenderItemColumn = (item: GeneralItem, index?: number, columnIn?: IColumn) => {
+   const onRenderItemColumn = (item: GeneralItem, _index?: number, columnIn?: IColumn) => {
 
       const column = columnIn!;
 
-      // simple cases
       switch (column.name) {
          case 'Name':
             return <Text variant="medium" styles={{ root: { fontFamily: "Horde Open Sans SemiBold" } }}>{item.name}</Text>
@@ -83,13 +84,22 @@ export const VersionModal: React.FC<{ show: boolean, onClose: () => void }> = ({
             if (item.value) {
                return <Text variant="medium" styles={{ root: { fontFamily: "Horde Open Sans SemiBold" } }}>{item.value}</Text>
             }
+            break;
+         case 'User ACLs':
+            if (item.userAcls === undefined) {
+               return null; // Plugin doesn't use ACLs
+            }
+            if (item.userAcls.length === 0) {
+               return <Text variant="small" styles={{ root: { fontFamily: "Horde Open Sans Regular", color: "#888" } }}>None</Text>
+            }
+            return <Text variant="small" styles={{ root: { fontFamily: "Horde Open Sans Regular" } }}>{item.userAcls.join(", ")}</Text>
       }
 
       return null;
    }
 
 
-   return <Modal className={hordeClasses.modal} isOpen={show} styles={{ main: { padding: 4, width: 640 } }} onDismiss={() => { if (show) close() }}>
+   return <Modal className={hordeClasses.modal} isOpen={show} styles={{ main: { padding: 4, width: 800 } }} onDismiss={() => { if (show) close() }}>
       <Stack tokens={{ childrenGap: 12 }}>
          <Stack horizontal styles={{ root: { paddingLeft: 12 } }}>
             <Stack grow style={{ paddingTop: 12, paddingLeft: 8 }}>
@@ -107,7 +117,7 @@ export const VersionModal: React.FC<{ show: boolean, onClose: () => void }> = ({
          <Stack tokens={{ childrenGap: 32 }} style={{ paddingLeft: 8 }}>
 
             <Stack horizontal>
-               <Stack style={{ minWidth: 600 }}>
+               <Stack style={{ minWidth: 760 }}>
                   <DetailsList
                      items={versionItems}
                      columns={columns}
diff --git a/Engine/Source/Programs/Horde/HordeServer.Shared/Plugins/IPluginDashboardAcls.cs b/Engine/Source/Programs/Horde/HordeServer.Shared/Plugins/IPluginDashboardAcls.cs
new file mode 100644
--- /dev/null
+++ b/Engine/Source/Programs/Horde/HordeServer.Shared/Plugins/IPluginDashboardAcls.cs
@@ -0,0 +1,32 @@
+// Copyright Epic Games, Inc. All Rights Reserved.
+
+using System.Security.Claims;
+using EpicGames.Horde.Acls;
+
+namespace HordeServer.Plugins
+{
+	/// <summary>
+	/// Interface for plugins to declare ACL actions relevant to dashboard visibility.
+	/// Plugins implementing this interface can specify what ACL actions control access
+	/// to their dashboard features, and the server will return which actions the current
+	/// user has when returning plugin information.
+	/// </summary>
+	public interface IPluginDashboardAcls
+	{
+		/// <summary>
+		/// Gets the ACL actions that are relevant for dashboard visibility for this plugin.
+		/// These actions will be checked against the current user and returned in the
+		/// server info response so the dashboard can filter plugins and mount points.
+		/// </summary>
+		/// <returns>Array of ACL actions to check for dashboard visibility</returns>
+		AclAction[] GetDashboardAclActions();
+
+		/// <summary>
+		/// Checks if the given user has the specified ACL action for this plugin
+		/// </summary>
+		/// <param name="action">The action to check</param>
+		/// <param name="user">The user to check authorization for</param>
+		/// <returns>True if the user has the action, false otherwise</returns>
+		bool AuthorizeDashboardAction(AclAction action, ClaimsPrincipal user);
+	}
+}
diff --git a/Engine/Source/Programs/Horde/HordeServer/Server/ServerController.cs b/Engine/Source/Programs/Horde/HordeServer/Server/ServerController.cs
--- a/Engine/Source/Programs/Horde/HordeServer/Server/ServerController.cs
+++ b/Engine/Source/Programs/Horde/HordeServer/Server/ServerController.cs
@@ -8,6 +8,7 @@ using System.Reflection;
 using System.Threading.Tasks;
 using EpicGames.Core;
 using EpicGames.Horde;
+using EpicGames.Horde.Acls;
 using EpicGames.Horde.Server;
 using HordeServer.Plugins;
 using HordeServer.Utilities;
@@ -77,12 +78,60 @@ namespace HordeServer.Server
 
 			response.Plugins = _pluginCollection.LoadedPlugins.ConvertAll(plugin => {
 				FileVersionInfo pluginVersion = FileVersionInfo.GetVersionInfo(plugin.Assembly.Location);
-				return new ServerPluginInfoResponse(plugin.Metadata.Name.ToString(), plugin.Metadata.Description, true, pluginVersion.ProductVersion ?? String.Empty);
+				string[]? userAcls = GetUserAclsForPlugin(plugin);
+				return new ServerPluginInfoResponse(plugin.Metadata.Name.ToString(), plugin.Metadata.Description, true, pluginVersion.ProductVersion ?? String.Empty, userAcls);
 			}).ToArray();
 
 			return response;
 		}
 
+		/// <summary>
+		/// Gets the list of ACL actions the current user has for a given plugin.
+		/// Returns null if plugin doesn't support dashboard ACLs, or an array (possibly empty)
+		/// of the user's granted ACL actions if it does.
+		/// </summary>
+		private string[]? GetUserAclsForPlugin(ILoadedPlugin plugin)
+		{
+			// Only return ACLs if the user is authenticated
+			if (User?.Identity?.IsAuthenticated != true)
+			{
+				return null;
+			}
+
+			// Check if the plugin config implements IPluginDashboardAcls
+			if (!_globalConfig.Value.Plugins.TryGetValue(plugin.Name, out IPluginConfig? pluginConfig))
+			{
+				return null;
+			}
+
+			if (pluginConfig is not IPluginDashboardAcls dashboardAcls)
+			{
+				return null;
+			}
+
+			// Get the dashboard-relevant ACL actions and check which ones the user has
+			AclAction[] actions = dashboardAcls.GetDashboardAclActions();
+
+			// Admin users get all dashboard ACL actions
+			if (User.HasAdminClaim())
+			{
+				return [.. actions.Select(a => a.Name)];
+			}
+
+			List<string> userAcls = [];
+
+			foreach (AclAction action in actions)
+			{
+				if (dashboardAcls.AuthorizeDashboardAction(action, User))
+				{
+					userAcls.Add(action.Name);
+				}
+			}
+
+			// Return the array (possibly empty) - this distinguishes from plugins that don't use ACLs (null)
+			return userAcls.ToArray();
+		}
+
 		/// <summary>
 		/// Gets connection information
 		/// </summary>
diff --git a/Engine/Source/Programs/Shared/EpicGames.Horde/Server/ServerMessages.cs b/Engine/Source/Programs/Shared/EpicGames.Horde/Server/ServerMessages.cs
--- a/Engine/Source/Programs/Shared/EpicGames.Horde/Server/ServerMessages.cs
+++ b/Engine/Source/Programs/Shared/EpicGames.Horde/Server/ServerMessages.cs
@@ -250,7 +250,8 @@ namespace EpicGames.Horde.Server
 	/// <param name="Description">Optional description of the plugin</param>
 	/// <param name="Loaded">Whether the plugin is loaded</param>
 	/// <param name="Version">The version of the plugin assembly</param>
-	public record ServerPluginInfoResponse(string Name, string? Description, bool Loaded, string? Version = null);
+	/// <param name="UserAcls">List of ACL actions the current user has for this plugin</param>
+	public record ServerPluginInfoResponse(string Name, string? Description, bool Loaded, string? Version = null, string[]? UserAcls = null);
 
 }
 
