Fix device not displaying when navigating phase history

#rnx
#Synced-CL 49844489

diff --git a/Engine/Source/Programs/Horde/HordeDashboard/plugins/test-automation/testAutomationEvents.tsx b/Engine/Source/Programs/Horde/HordeDashboard/plugins/test-automation/testAutomationEvents.tsx
--- a/Engine/Source/Programs/Horde/HordeDashboard/plugins/test-automation/testAutomationEvents.tsx
+++ b/Engine/Source/Programs/Horde/HordeDashboard/plugins/test-automation/testAutomationEvents.tsx
@@ -259,9 +259,7 @@ export const PhaseEventsPane: React.FC<{phase: TestPhaseStatus, onEventsLoaded?:
         return false;
     }, [filterError, filterWarning, filterInfo]);
 
-    const uniqueDevices = useMemo(() => {
-        return phase.devices.filter((device, index) => phase.devices.findIndex(other => device.isSame(other)) === index);
-    },[phase]);
+    const uniqueDevices = phase.devices.filter((device, index) => phase.devices.findIndex(other => device.isSame(other)) === index);
 
     const borderColor = dashboard.darktheme ? "#4D4C4B" : "#6D6C6B";
 
