[Horde][Server] Exposing existing Session ID to the telemetry api

#Synced-CL 49813167

diff --git a/Engine/Source/Programs/Horde/Plugins/Analytics/HordeServer.Analytics/Telemetry/TelemetryController.cs b/Engine/Source/Programs/Horde/Plugins/Analytics/HordeServer.Analytics/Telemetry/TelemetryController.cs
--- a/Engine/Source/Programs/Horde/Plugins/Analytics/HordeServer.Analytics/Telemetry/TelemetryController.cs
+++ b/Engine/Source/Programs/Horde/Plugins/Analytics/HordeServer.Analytics/Telemetry/TelemetryController.cs
@@ -160,11 +160,12 @@ namespace HordeServer.Telemetry
 		/// <param name="appVersion">Version number of the application</param>
 		/// <param name="appEnvironment">Name of the environment that the sending application is running in</param>
 		/// <param name="uploadType">Type of data being uploaded</param>
+		/// <param name="sessionId">Identifier used to group telemetry within a single session.</param>
 		[HttpPost]
 		[Route("/api/v1/telemetry")]
-		public ActionResult PostEvent([FromBody] PostTelemetryEventStreamRequest request, [FromQuery][Required] string appId, [FromQuery][Required] string appVersion, [FromQuery][Required] string appEnvironment, [FromQuery][Required] TelemetryUploadType uploadType)
+		public ActionResult PostEvent([FromBody] PostTelemetryEventStreamRequest request, [FromQuery][Required] string appId, [FromQuery][Required] string appVersion, [FromQuery][Required] string appEnvironment, [FromQuery][Required] TelemetryUploadType uploadType, [FromQuery] string? sessionId = null)
 		{
-			return PostEvent(TelemetryStoreId.Default, request, appId, appVersion, appEnvironment, uploadType);
+			return PostEvent(TelemetryStoreId.Default, request, appId, appVersion, appEnvironment, uploadType, sessionId);
 		}
 
 		/// <summary>
@@ -176,16 +177,17 @@ namespace HordeServer.Telemetry
 		/// <param name="appVersion">Version number of the application</param>
 		/// <param name="appEnvironment">Name of the environment that the sending application is running in</param>
 		/// <param name="uploadType">Type of data being uploaded</param>
+		/// <param name="sessionId">Identifier used to group telemetry within a single session.</param>
 		[HttpPost]
 		[Route("/api/v1/telemetry/{storeId}")]
-		public ActionResult PostEvent(TelemetryStoreId storeId, [FromBody] PostTelemetryEventStreamRequest request, [FromQuery][Required] string appId, [FromQuery][Required] string appVersion, [FromQuery][Required] string appEnvironment, [FromQuery][Required] TelemetryUploadType uploadType)
+		public ActionResult PostEvent(TelemetryStoreId storeId, [FromBody] PostTelemetryEventStreamRequest request, [FromQuery][Required] string appId, [FromQuery][Required] string appVersion, [FromQuery][Required] string appEnvironment, [FromQuery][Required] TelemetryUploadType uploadType, [FromQuery] string? sessionId = null)
 		{
 			if (uploadType != TelemetryUploadType.EtEventStream)
 			{
 				return BadRequest("Invalid upload type");
 			}
 
-			TelemetryRecordMeta recordMeta = new TelemetryRecordMeta(AppId: appId, AppVersion: appVersion, AppEnvironment: appEnvironment);
+			TelemetryRecordMeta recordMeta = new TelemetryRecordMeta(AppId: appId, AppVersion: appVersion, AppEnvironment: appEnvironment, SessionId: sessionId);
 			foreach (JsonElement eventPayload in request.Events)
 			{
 				_telemetryManager.WriteEvent(storeId, recordMeta, eventPayload);
@@ -247,4 +249,4 @@ namespace HordeServer.Telemetry
 			return responses;
 		}
 	}
-}
+}
\ No newline at end of file
diff --git a/Engine/Source/Programs/Shared/EpicGames.Horde/HordeHttpClient.cs b/Engine/Source/Programs/Shared/EpicGames.Horde/HordeHttpClient.cs
--- a/Engine/Source/Programs/Shared/EpicGames.Horde/HordeHttpClient.cs
+++ b/Engine/Source/Programs/Shared/EpicGames.Horde/HordeHttpClient.cs
@@ -437,10 +437,11 @@ namespace EpicGames.Horde
 		/// <param name="appVersion">The version of the App.</param>
 		/// <param name="appEnvironment">The environment.</param>
 		/// <param name="payloads">A list of payloads to write as telemetry events.</param>
+		/// <param name="sessionId">The session id which this telemetry is a part of.</param>
 		/// <param name="cancellationToken">The cancellation token to use throughout.</param>
 		/// <returns>The HttpResponseMessage.</returns>
 		/// <remarks>BadRequest can be returned under circusmtances where the payload is not serializable (primitive or null payload, <see cref="NotSupportedException"/>, and <see cref="JsonException"/>).</remarks>
-		public async Task<HttpResponseMessage> PostTelemetryAsync(string appId, int appVersion, string appEnvironment, object[] payloads, CancellationToken cancellationToken = default)
+		public async Task<HttpResponseMessage> PostTelemetryAsync(string appId, int appVersion, string appEnvironment, object[] payloads, string? sessionId = null, CancellationToken cancellationToken = default)
 		{
 			QueryStringBuilder queryParams = new QueryStringBuilder();
 			queryParams.Add("appId", appId);
@@ -448,6 +449,11 @@ namespace EpicGames.Horde
 			queryParams.Add("appEnvironment", appEnvironment);
 			queryParams.Add("uploadType", "EtEventStream");
 
+			if (!String.IsNullOrEmpty(sessionId))
+			{
+				queryParams.Add("SessionId", sessionId);
+			}
+
 			PostTelemetryEventStreamRequest request = new PostTelemetryEventStreamRequest();
 
 			try
@@ -467,7 +473,7 @@ namespace EpicGames.Horde
 					}
 				}
 			}
-			catch(NotSupportedException) // The scenario where a type doesn't contain a JsonConverter.
+			catch (NotSupportedException) // The scenario where a type doesn't contain a JsonConverter.
 			{
 				return new HttpResponseMessage(HttpStatusCode.BadRequest);
 			}
@@ -607,7 +613,7 @@ namespace EpicGames.Horde
 			CommitId? maxChange = null,
 			DateTimeOffset? minCreateTime = null,
 			DateTimeOffset? maxCreateTime = null,
-			bool includePreflight = false,			
+			bool includePreflight = false,
 			int count = 100,
 			CancellationToken cancellationToken = default)
 		{
@@ -649,7 +655,7 @@ namespace EpicGames.Horde
 
 			return GetAsync<List<GetJobResponse>>(_httpClient, queryBuilder.ToString(), cancellationToken);
 		}
-		
+
 		#endregion
 
 		#region Log
