Horde: Fix compiler warnings for JobDriver.Tests after upgrading MSTest

diff --git a/Engine/Source/Programs/Horde/Drivers/JobDriver.Tests/Execution/WorkspaceExecutorTest.cs b/Engine/Source/Programs/Horde/Drivers/JobDriver.Tests/Execution/WorkspaceExecutorTest.cs
--- a/Engine/Source/Programs/Horde/Drivers/JobDriver.Tests/Execution/WorkspaceExecutorTest.cs
+++ b/Engine/Source/Programs/Horde/Drivers/JobDriver.Tests/Execution/WorkspaceExecutorTest.cs
@@ -242,7 +242,7 @@ namespace JobDriver.Tests.Execution
 			using WorkspaceExecutor executor = new(executorOptions, _workspace, null, SimpleTestExecutor.NoOpTracer, NullLogger.Instance);
 
 			RpcBeginBatchResponse batch = new RpcBeginBatchResponse { };
-			await Assert.ThrowsExceptionAsync<WorkspaceMaterializationException>(() => executor.InitializeAsync(batch, _logger, CancellationToken.None));
+			await Assert.ThrowsExactlyAsync<WorkspaceMaterializationException>(() => executor.InitializeAsync(batch, _logger, CancellationToken.None));
 		}
 
 		private static void AssertWorkspaceFile(IWorkspaceMaterializer workspace, string relativePath, string expectedContent)
diff --git a/Engine/Source/Programs/Horde/Drivers/JobDriver.Tests/JobDriver.Tests.csproj b/Engine/Source/Programs/Horde/Drivers/JobDriver.Tests/JobDriver.Tests.csproj
--- a/Engine/Source/Programs/Horde/Drivers/JobDriver.Tests/JobDriver.Tests.csproj
+++ b/Engine/Source/Programs/Horde/Drivers/JobDriver.Tests/JobDriver.Tests.csproj
@@ -18,7 +18,10 @@
 
   <ItemGroup>
     <PackageReference Include="Microsoft.Extensions.Logging.Abstractions" Version="8.0.3" />
+    <PackageReference Include="Microsoft.NET.Test.Sdk" Version="17.14.1" />
     <PackageReference Include="Moq" Version="4.20.72" />
+    <PackageReference Include="MSTest.TestAdapter" Version="3.11.1" />
+    <PackageReference Include="MSTest.TestFramework" Version="3.11.1" />
   </ItemGroup>
 
   <ItemGroup>
diff --git a/Engine/Source/Programs/Horde/Drivers/JobDriver.Tests/WorkerServiceTest.cs b/Engine/Source/Programs/Horde/Drivers/JobDriver.Tests/WorkerServiceTest.cs
--- a/Engine/Source/Programs/Horde/Drivers/JobDriver.Tests/WorkerServiceTest.cs
+++ b/Engine/Source/Programs/Horde/Drivers/JobDriver.Tests/WorkerServiceTest.cs
@@ -82,7 +82,7 @@ namespace JobDriver.Tests
 					return JobStepOutcome.Success;
 				});
 
-				await Assert.ThrowsExceptionAsync<TaskCanceledException>(() => executor.ExecuteStepAsync(
+				await Assert.ThrowsExactlyAsync<TaskCanceledException>(() => executor.ExecuteStepAsync(
 					null!, NullLogger.Instance, cancelSource.Token, stepCancelSource.Token));
 			}
 
diff --git a/Engine/Source/Programs/Horde/Drivers/JobDriver.Tests/WorkspaceExecutorTests.cs b/Engine/Source/Programs/Horde/Drivers/JobDriver.Tests/WorkspaceExecutorTests.cs
--- a/Engine/Source/Programs/Horde/Drivers/JobDriver.Tests/WorkspaceExecutorTests.cs
+++ b/Engine/Source/Programs/Horde/Drivers/JobDriver.Tests/WorkspaceExecutorTests.cs
@@ -61,7 +61,7 @@ public class WorkspaceExecutorTests
 	[TestMethod]
 	public void Factory()
 	{
-		Assert.ThrowsException<ArgumentException>(() => GetMaterializer(new RpcJobOptions(), new RpcAgentWorkspace(), []));
+		Assert.ThrowsExactly<ArgumentException>(() => GetMaterializer(new RpcJobOptions(), new RpcAgentWorkspace(), []));
 		
 		AssertWm(ManagedWorkspace, "", "");
 		AssertWm(ManagedWorkspace, "mAnaGedWorkSpace", "");
