[Horde][Server] Removal of sparse constraint as it's less effective on queryPlanner version 1

#Synced-CL 49874115

diff --git a/Engine/Source/Programs/Horde/Plugins/Build/HordeServer.Build/Jobs/JobStepRefCollection.cs b/Engine/Source/Programs/Horde/Plugins/Build/HordeServer.Build/Jobs/JobStepRefCollection.cs
--- a/Engine/Source/Programs/Horde/Plugins/Build/HordeServer.Build/Jobs/JobStepRefCollection.cs
+++ b/Engine/Source/Programs/Horde/Plugins/Build/HordeServer.Build/Jobs/JobStepRefCollection.cs
@@ -164,8 +164,7 @@ namespace HordeServer.Jobs
 			// Using this index, we can cost-effectively query:
 			//	- All job step refs belonging to a specific job
 			//	- All job step refs belonging to a batch (of a job)
-			// It's sparse in nature due for backwards compatibility, and it is accepted that some items which use this index will miss items.
-			indexes.Add(keys => keys.Ascending(x => x.JobId).Ascending(x => x.BatchId).Ascending(x => x.StepId), sparse: true);
+			indexes.Add(keys => keys.Ascending(x => x.JobId).Ascending(x => x.BatchId).Ascending(x => x.StepId));
 
 			_jobStepRefs = mongoService.GetCollection<JobStepRef>("JobStepRefs", indexes);
 			_telemetryWriter = telemetryWriter;
