Capture user's test metadata updates in test audit log

#rnx

#Synced-CL 49884572

diff --git a/Engine/Source/Programs/Horde/Plugins/Build/HordeServer.Build/Jobs/TestData/TestDataController.cs b/Engine/Source/Programs/Horde/Plugins/Build/HordeServer.Build/Jobs/TestData/TestDataController.cs
--- a/Engine/Source/Programs/Horde/Plugins/Build/HordeServer.Build/Jobs/TestData/TestDataController.cs
+++ b/Engine/Source/Programs/Horde/Plugins/Build/HordeServer.Build/Jobs/TestData/TestDataController.cs
@@ -10,9 +10,11 @@ using EpicGames.Horde.Streams;
 using EpicGames.Horde.Users;
 using HordeServer.Auditing;
 using HordeServer.Streams;
+using HordeServer.Users;
 using HordeServer.Utilities;
 using Microsoft.AspNetCore.Authorization;
 using Microsoft.AspNetCore.Mvc;
+using Microsoft.Extensions.Logging;
 using Microsoft.Extensions.Options;
 using MongoDB.Bson;
 using MongoDB.Bson.Serialization;
@@ -37,6 +39,8 @@ namespace HordeServer.Jobs.TestData
 		/// </summary>
 		private readonly ITestDataCollection _testDataCollection;
 
+		private readonly IUserCollection _userCollection;
+
 		readonly TestDataService _testDataService;
 
 		readonly IOptionsSnapshot<BuildConfig> _buildConfig;
@@ -44,12 +48,13 @@ namespace HordeServer.Jobs.TestData
 		/// <summary>
 		/// Constructor
 		/// </summary>
-		public TestDataController(TestDataService testDataService, JobService jobService, ITestDataCollection testDataCollection, IOptionsSnapshot<BuildConfig> buildConfig)
+		public TestDataController(TestDataService testDataService, JobService jobService, ITestDataCollection testDataCollection, IOptionsSnapshot<BuildConfig> buildConfig, IUserCollection userCollection)
 		{
 			_jobService = jobService;
 			_testDataCollection = testDataCollection;
 			_testDataService = testDataService;
 			_buildConfig = buildConfig;
+			_userCollection = userCollection;
 		}
 
 		/// <summary>
@@ -935,26 +940,37 @@ namespace HordeServer.Jobs.TestData
 		[ProducesResponseType(typeof(GetTestAuditResponse), 200)]
 		public async Task<ActionResult<GetTestAuditResponse>> AddTestAuditV2Async([FromBody] AddTestAuditRequest request, CancellationToken cancellationToken = default)
 		{
-			GetTestAuditResponse responses = new GetTestAuditResponse();
-
 			UserId? userId = User.GetUserId();
-			if (userId == null)
+			if (userId is null)
 			{
 				return BadRequest("No user logged in.");
 			}
 
-			if (request.TestId == null || request.TestId.Length == 0)
+			TestId testId;
+			if (request.TestId is null || request.TestId.Length == 0 || !TestId.TryParse(request.TestId, out testId))
 			{
 				return BadRequest("Invalid test id.");
 			}
 
+			ITestAudit? prevAudit = (await _testDataService.FindTestAuditsV2Async([testId], null, cancellationToken)).ToList().FirstOrDefault();
+
+			// add notification settings hidden from user
+			Dictionary<string, string>? notificationsHiddenFromUser = prevAudit?.Notifications?
+				.Where(n => _buildConfig.Value.TryGetStream(n.Stream, out StreamConfig? streamConfig) && !streamConfig.Authorize(JobAclAction.ViewJob, User))
+				.ToDictionary(n => n.Stream.ToString(), n => n.Workflow.ToString());
+			if (notificationsHiddenFromUser is not null && notificationsHiddenFromUser.Count > 0)
+			{
+				request.Notifications = request.Notifications is null? notificationsHiddenFromUser : request.Notifications.Concat(notificationsHiddenFromUser).ToLookup(n => n.Key, n => n.Value).ToDictionary(n => n.Key, n => n.Last());
+			}
+
+			// add audit entry
 			TestAuditOptions options = new TestAuditOptions()
 			{
 				LastAuditDate = request.LastAuditDate?.UtcDateTime,
 				IsUnderAudit = request.IsUnderAudit,
 				EnableSchedule = request.EnableSchedule ?? true,
-				Owner = request.Owner == null || request.Owner.Length == 0? null : UserId.Parse(request.Owner),
-				Customers = request.Customers?.Where(u => u.Length > 0).Select(u => UserId.Parse(u)).ToList(),
+				Owner = request.Owner == null || request.Owner.Length == 0 || !UserId.TryParse(request.Owner, out _) ? null : UserId.Parse(request.Owner),
+				Customers = request.Customers?.Where(u => u.Length > 0 && UserId.TryParse(u, out _)).Select(u => UserId.Parse(u)).ToList(),
 				Notifications = request.Notifications?.Select(
 					n => new TestNotificationSettingsOptions()
 					{
@@ -968,11 +984,82 @@ namespace HordeServer.Jobs.TestData
 				Notes = request.Notes
 			};
 
-			ITestAudit audit = await _testDataService.AddTestAuditAsync(userId!.Value, TestId.Parse(request.TestId), options, cancellationToken);
+			ITestAudit audit = await _testDataService.AddTestAuditAsync(userId!.Value, testId, options, cancellationToken);
+
+			// log changes
+			await LogAuditChangesAsync(userId!.Value, testId, options, prevAudit, cancellationToken);
 
 			return GetNewTestAutditResponse(audit);
 		}
 
+		private async Task LogAuditChangesAsync(UserId userId, TestId testId, TestAuditOptions options, ITestAudit? prevAudit, CancellationToken cancellationToken = default)
+		{
+			string userName = (await _userCollection.GetCachedUserAsync(userId, cancellationToken))?.Name ?? userId.ToString();
+			IAuditLogChannel<TestId> logger = _testDataService.GetTestAuditLogger(testId);
+
+			if (options.Team != prevAudit?.Team)
+			{
+				logger.LogInformation("Team changed to {Team} by {UserName}", options.Team ?? "none", userName);
+			}
+			if (options.Intent != prevAudit?.Intent)
+			{
+				logger.LogInformation("Intent changed to {Intent} by {UserName}", options.Intent ?? "none", userName);
+			}
+			if (options.Harness != prevAudit?.Harness)
+			{
+				logger.LogInformation("Harness changed to {Harness} by {UserName}", options.Harness ?? "none", userName);
+			}
+			if (options.Profile != prevAudit?.Profile)
+			{
+				logger.LogInformation("Profile changed to {Profile} by {UserName}", options.Profile ?? "none", userName);
+			}
+			if (options.EnableSchedule != prevAudit?.EnableSchedule)
+			{
+				logger.LogInformation("Enable Schedule changed to {EnableSchedule} by {UserName}", options.EnableSchedule, userName);
+			}
+			if (options.LastAuditDate != prevAudit?.LastAuditDate)
+			{
+				logger.LogInformation("Last Audit Date changed to {LastAuditDate} by {UserName}", options.LastAuditDate?.ToString() ?? "none", userName);
+			}
+			if (options.IsUnderAudit != prevAudit?.IsUnderAudit)
+			{
+				logger.LogInformation("Is Under Audit changed to {IsUnderAudit} by {UserName}", options.IsUnderAudit?.ToString() ?? "none", userName);
+			}
+			if (options.Owner != prevAudit?.Owner)
+			{
+				string ownerName = options.Owner is null ? "none" : (await _userCollection.GetCachedUserAsync(options.Owner.Value, cancellationToken))?.Name ?? options.Owner.Value.ToString();
+				logger.LogInformation("Owner changed to {Owner} by {UserName}", ownerName, userName);
+			}
+			if (options.Customers is null && prevAudit?.Customers is not null)
+			{
+				logger.LogInformation("Customers changed to {Customers} by {UserName}", "none", userName);
+			}
+			else if (options.Customers is not null && (prevAudit?.Customers is null || !options.Customers.SequenceEqual(prevAudit.Customers)))
+			{
+				IEnumerable<Task<IUser?>> tasks = options.Customers.Select(c => _userCollection.GetCachedUserAsync(c, cancellationToken).AsTask());
+				string[] customerNames = (await Task.WhenAll(tasks)).OfType<IUser>().Select(c => c.Name).ToArray();
+				logger.LogInformation("Customers changed to {Customers} by {UserName}", String.Join(", ", customerNames), userName);
+			}
+			if (options.Notifications is null && prevAudit?.Notifications is not null)
+			{
+				logger.LogInformation("Notifications changed to {Notifications} by {UserName}", "none", userName);
+			}
+			else if (options.Notifications is not null
+					&& (
+						prevAudit?.Notifications is null
+						|| options.Notifications.Any(n => !prevAudit.Notifications.Any(p => p.Stream == n.Stream && p.Workflow == n.Workflow))
+						|| prevAudit.Notifications.Any(n => !options.Notifications.Any(p => p.Stream == n.Stream && p.Workflow == n.Workflow))
+						))
+			{
+				string[] notifications = options.Notifications.Select(n => $"[{n.Stream}: {n.Workflow}]").ToArray();
+				logger.LogInformation("Notifications changed to {Notifications} by {UserName}", String.Join(", ", notifications), userName);
+			}
+			if (options.Notes != prevAudit?.Notes)
+			{
+				logger.LogInformation("Notes changed by {UserName} to:\n{Notes}", userName, options.Notes ?? "<none>");
+			}
+		}
+
 		/// <summary>
 		/// Get test audit history
 		/// </summary>
