Horde: Much better dashboard commit cache behavior

diff --git a/Engine/Source/Programs/Horde/HordeDashboard/src/backend/CommitCache.ts b/Engine/Source/Programs/Horde/HordeDashboard/src/backend/CommitCache.ts
--- a/Engine/Source/Programs/Horde/HordeDashboard/src/backend/CommitCache.ts
+++ b/Engine/Source/Programs/Horde/HordeDashboard/src/backend/CommitCache.ts
@@ -47,48 +47,22 @@ export class CommitCache {
             return;
         }
 
-        const cminCL = 999999999999;
-        const cmaxCL = -1;
-        const maxDelta = 131072;
+        // Filter to only changes we don't already have cached
+        const uncachedChanges = changes.filter(c => !map!.get(c));
 
-        let minCL = cminCL;
-        let maxCL = cmaxCL;
-
-        changes.forEach(c => {
-            const change = map!.get(c);
-
-            if (!change) {
-
-                if (c < minCL && (minCL === cminCL || Math.abs(minCL - c) < maxDelta)) {
-                    minCL = c;
-                }
-
-                if (c > maxCL && (maxCL === cmaxCL || Math.abs(maxCL - c) < maxDelta)) {
-                    maxCL = c;
-                }
-            }
-        })
-
-        if (maxCL === cmaxCL && minCL === cminCL) {
+        if (!uncachedChanges.length) {
             return;
         }
 
-        changes.forEach(c => {
-
-            const change = map!.get(c);
-            if (change) {
-                return;
-            }
-
+        // Set placeholders for changes we're about to fetch
+        uncachedChanges.forEach(c => {
             map!.set(c, {
-                number: 0, description: "", dateUtc: "", authorInfo : { id : "", email:"", name: ""}
-            })
-
-        })
+                number: 0, description: "", dateUtc: "", authorInfo: { id: "", email: "", name: "" }
+            });
+        });
 
-        // new query
         try {
-            const commits = await backend.getChangeSummaries(streamId, minCL, maxCL, undefined, undefined, maxCL - minCL + 1);
+            const commits = await backend.getBatchChangeSummaries(streamId, uncachedChanges);
             commits.forEach(c => map!.set(c.number, c));
             this.setUpdated();
         } catch (reason) {
diff --git a/Engine/Source/Programs/Horde/HordeDashboard/src/backend/JobHandler.ts b/Engine/Source/Programs/Horde/HordeDashboard/src/backend/JobHandler.ts
--- a/Engine/Source/Programs/Horde/HordeDashboard/src/backend/JobHandler.ts
+++ b/Engine/Source/Programs/Horde/HordeDashboard/src/backend/JobHandler.ts
@@ -4,6 +4,7 @@ import { action, makeObservable, observable } from 'mobx';
 import moment from 'moment';
 import backend from '../backend';
 import { JobData, JobState, JobStepOutcome, JobStreamQuery, StreamData } from '../backend/Api';
+import { CommitCache } from './CommitCache';
 
 export type FilterStatus = "Running" | "Complete" | "Succeeded" | "Failed" | "Cancelled" | "Waiting";
 
@@ -12,10 +13,11 @@ const jobsRefreshTime = 10000;
 
 export class JobHandler {
 
-    constructor(includeBatches: boolean = false, jobLimit?: number) {
+    constructor(includeBatches: boolean = false, jobLimit?: number, commitCache?: CommitCache) {
         makeObservable(this);
         this.includeBatches = includeBatches;
         this.jobLimit = jobLimit;
+        this.commitCache = commitCache;
     }
 
 
@@ -172,6 +174,12 @@ export class JobHandler {
                 wasUpdated = true;
             }
 
+            // Populate commit cache with change numbers from jobs
+            if (this.commitCache && this.stream) {
+                const changes = [...new Set(this.jobs.map(j => j.change).filter(change => !!change))] as number[];
+                await this.commitCache.set(this.stream.id, changes);
+            }
+
         } catch (reason) {
             console.log(reason);
         } finally {
@@ -299,6 +307,7 @@ export class JobHandler {
     private preflightStartedByUserId?: string;
 
     private jobLimit?: number;
+    private commitCache?: CommitCache;
 
     initial = true;
 
diff --git a/Engine/Source/Programs/Horde/HordeDashboard/src/components/JobView.tsx b/Engine/Source/Programs/Horde/HordeDashboard/src/components/JobView.tsx
--- a/Engine/Source/Programs/Horde/HordeDashboard/src/components/JobView.tsx
+++ b/Engine/Source/Programs/Horde/HordeDashboard/src/components/JobView.tsx
@@ -7,6 +7,7 @@ import React, { useEffect, useState } from 'react';
 import { Link, useParams } from 'react-router-dom';
 import { useBackend } from '../backend';
 import { GetJobsTabLabelColumnResponse, GetJobsTabParameterColumnResponse, GetLabelStateResponse, JobData, JobsTabColumnType, JobsTabData, JobState, LabelOutcome, LabelState, StreamData } from '../backend/Api';
+import { CommitCache } from '../backend/CommitCache';
 import dashboard from '../backend/Dashboard';
 import { JobHandler } from '../backend/JobHandler';
 import { filterJob, JobFilterSimple } from '../base/utilities/filter';
@@ -174,7 +175,8 @@ const sortJobGroups = (items: JobItem[]): [JobItem[], JobGroup[]] => {
    return [nitems, groups];
 };
 
-const jobHandler = new JobHandler();
+const commitCache = new CommitCache();
+const jobHandler = new JobHandler(false, undefined, commitCache);
 
 const JobList: React.FC<{ tab: string; filter: JobFilterSimple }> = observer(({ tab, filter }) => {
 
@@ -199,6 +201,7 @@ const JobList: React.FC<{ tab: string; filter: JobFilterSimple }> = observer(({
 
    // @todo: this pattern is becoming annoying, have to reference once in render for observable to see
    if (jobHandler.updated) { }
+   if (commitCache.updated) { }
 
    if (!streamId || !stream || !projectStore.streamById(streamId)) {
       console.error("bad stream id setting up JobList");
@@ -437,7 +440,7 @@ const JobList: React.FC<{ tab: string; filter: JobFilterSimple }> = observer(({
 
    };
 
-   function onRenderItemColumnInner(item: JobItem, index?: number, column?: IColumn) {
+   function onRenderItemColumnInner(item: JobItem, _index?: number, column?: IColumn) {
 
       const fieldContent = item[column!.fieldName as keyof JobItem] as string;
 
@@ -480,10 +483,12 @@ const JobList: React.FC<{ tab: string; filter: JobFilterSimple }> = observer(({
 
       }
 
+      const commit = streamId ? commitCache.getCommit(streamId, item.job.change!) : undefined;
+
       switch (column!.key) {
 
          case 'Change':
-            return <ChangeButton job={item.job} />;
+            return <ChangeButton job={item.job} commit={commit} />;
 
          case 'Status':
             return <Stack style={style} verticalAlign="center" verticalFill={true} horizontalAlign={"center"}>{renderDetailStatus(item.job)}</Stack>;
diff --git a/Engine/Source/Programs/Horde/HordeDashboard/src/components/JobViewAll.tsx b/Engine/Source/Programs/Horde/HordeDashboard/src/components/JobViewAll.tsx
--- a/Engine/Source/Programs/Horde/HordeDashboard/src/components/JobViewAll.tsx
+++ b/Engine/Source/Programs/Horde/HordeDashboard/src/components/JobViewAll.tsx
@@ -7,6 +7,7 @@ import React, { useEffect } from "react";
 import { Link, useParams } from "react-router-dom";
 import { useBackend } from "../backend";
 import { GetLabelStateResponse, GetStepResponse, JobData, JobState, JobStepOutcome, LabelOutcome, LabelState, StepData, StreamData } from "../backend/Api";
+import { CommitCache } from "../backend/CommitCache";
 import dashboard from "../backend/Dashboard";
 import { JobHandler } from "../backend/JobHandler";
 import { filterJob, JobFilterSimple } from "../base/utilities/filter";
@@ -55,7 +56,8 @@ const customStyles = mergeStyleSets({
 });
 
 
-const jobHandler = new JobHandler(true, 300);
+const commitCache = new CommitCache();
+const jobHandler = new JobHandler(true, 300, commitCache);
 
 const JobFilterPanel: React.FC = () => {
 
@@ -85,6 +87,7 @@ const JobViewAllInner: React.FC<{ filter: JobFilterSimple }> = observer(({ filte
    // subscribe
    if (jobHandler.updated) { }
    if (jobFilter.updated) { }
+   if (commitCache.updated) { }
 
    const stream = projectStore.streamById(streamId);
 
@@ -291,8 +294,8 @@ const JobViewAllInner: React.FC<{ filter: JobFilterSimple }> = observer(({ filte
       }
 
       if (column!.name === "Change") {
-
-         return <ChangeButton job={job} />;
+         const commit = streamId ? commitCache.getCommit(streamId, job.change!) : undefined;
+         return <ChangeButton job={job} commit={commit} />;
 
       }
 
diff --git a/Engine/Source/Programs/Horde/HordeDashboard/src/components/JobViewIncremental.tsx b/Engine/Source/Programs/Horde/HordeDashboard/src/components/JobViewIncremental.tsx
--- a/Engine/Source/Programs/Horde/HordeDashboard/src/components/JobViewIncremental.tsx
+++ b/Engine/Source/Programs/Horde/HordeDashboard/src/components/JobViewIncremental.tsx
@@ -43,8 +43,8 @@ type JobGroup = {
 }
 
 const commitCache = new CommitCache();
+const jobHandler = new JobHandler(false, undefined, commitCache);
 
- 
 const buildColumns = (jobTab: JobsTabData, streamId: string): IColumn[] => {
 
    const startedByName = jobTab.showNames ? "Started By" : "Change";
@@ -178,8 +178,6 @@ const sortJobGroups = (items: JobItem[]): [JobItem[], JobGroup[]] => {
    return [nitems, groups];
 };
 
-const jobHandler = new JobHandler();
-
 const JobList: React.FC<{ tab: string; filter: JobFilterSimple }> = observer(({ tab, filter }) => {
 
    const { streamId } = useParams<{ streamId: string }>();
@@ -243,9 +241,6 @@ const JobList: React.FC<{ tab: string; filter: JobFilterSimple }> = observer(({
       return filterJob(j, filter.filterKeyword, additionKeywords)
    });
 
-   const changes = jobs.map(j => j.change).filter(change => !!change) as number[];
-   commitCache.set(streamId, changes);
-
    jobItems = jobs.map(j => {
 
       let startedBy = j.startedByUserInfo?.name ?? "Scheduler";
