Horde: Add diagnostic logging to schedule gate job searches

Track why jobs are skipped (failed, incomplete, or missing target) and log summary when all candidates are exhausted.

diff --git a/Engine/Source/Programs/Horde/Plugins/Build/HordeServer.Build/Jobs/Schedules/ScheduleService.cs b/Engine/Source/Programs/Horde/Plugins/Build/HordeServer.Build/Jobs/Schedules/ScheduleService.cs
--- a/Engine/Source/Programs/Horde/Plugins/Build/HordeServer.Build/Jobs/Schedules/ScheduleService.cs
+++ b/Engine/Source/Programs/Horde/Plugins/Build/HordeServer.Build/Jobs/Schedules/ScheduleService.cs
@@ -586,7 +586,12 @@ namespace HordeServer.Jobs.Schedules
 			span.SetAttribute("horde.gate.target", gate.Target);
 			span.SetAttribute("horde.min_commit_id", minCommitId?.ToString());
 			span.SetAttribute("horde.max_commit_id", maxCommitId?.ToString());
-			
+
+			int jobsEvaluated = 0;
+			int jobsSkippedFailed = 0;
+			int jobsSkippedNotComplete = 0;
+			int jobsSkippedNoTarget = 0;
+
 			for (int c = 0; ; c++)
 			{
 				cancellationToken.ThrowIfCancellationRequested();
@@ -599,22 +604,50 @@ namespace HordeServer.Jobs.Schedules
 				IJob? job = jobs.FirstOrDefault(x => maxCommitId == null || x.CommitId < maxCommitId);
 				if (job == null)
 				{
-					span.AddEvent("Job not found");
+					span.AddEvent("Gate search exhausted", new SpanAttributes(new KeyValuePair<string, object?>[]
+					{
+						new("horde.jobs_evaluated", jobsEvaluated),
+						new("horde.jobs_skipped_failed", jobsSkippedFailed),
+						new("horde.jobs_skipped_not_complete", jobsSkippedNotComplete),
+						new("horde.jobs_skipped_no_target", jobsSkippedNoTarget)
+					}));
+
+					if (jobsEvaluated > 0)
+					{
+						scheduleLogger.LogInformation("Gate search for {StreamId} template {TemplateId} (gate: {GateTemplateId}/{GateTarget}) exhausted after evaluating {JobsEvaluated} jobs: {Failed} failed, {NotComplete} not complete, {NoTarget} missing target",
+							streamId, templateRefId, gate.TemplateId, gate.Target, jobsEvaluated, jobsSkippedFailed, jobsSkippedNotComplete, jobsSkippedNoTarget);
+					}
 					return null;
 				}
 
+				jobsEvaluated++;
+
 				IGraph? graph = await _graphs.GetAsync(job.GraphHash, cancellationToken);
 				if (graph != null)
 				{
 					(JobStepState stepState, JobStepOutcome stepOutcome)? state = job.GetTargetState(graph, gate.Target);
-					if (state is { stepState: JobStepState.Completed })
+
+					if (state == null)
+					{
+						jobsSkippedNoTarget++;
+						_logger.LogDebug("Gate: Job {JobId} (CL {CommitId}) - target '{Target}' not found in graph or missing steps", job.Id, job.CommitId, gate.Target);
+					}
+					else if (state is { stepState: JobStepState.Completed })
 					{
 						if (state is { stepOutcome: JobStepOutcome.Success or JobStepOutcome.Warnings})
 						{
+							span.SetAttribute("horde.jobs_evaluated", jobsEvaluated);
 							return job.CommitId;
 						}
+						jobsSkippedFailed++;
+						_logger.LogDebug("Gate: Job {JobId} (CL {CommitId}) - target '{Target}' completed with {Outcome}", job.Id, job.CommitId, gate.Target, state.Value.stepOutcome);
 						scheduleLogger.LogInformation("Skipping trigger of {StreamName} template {TemplateId} - last {OtherTemplateRefId} job ({JobId}) ended with errors", streamId, templateRefId, gate.TemplateId, job.Id);
 					}
+					else
+					{
+						jobsSkippedNotComplete++;
+						_logger.LogDebug("Gate: Job {JobId} (CL {CommitId}) - target '{Target}' not complete (state: {State})", job.Id, job.CommitId, gate.Target, state.Value.stepState);
+					}
 				}
 
 				maxCommitId = job.CommitId;
