Add test health audit log
+ add support to generic audit type log

#rnx
#Synced-CL 49884572

diff --git a/Engine/Source/Programs/Horde/HordeDashboard/plugins/test-automation/index.tsx b/Engine/Source/Programs/Horde/HordeDashboard/plugins/test-automation/index.tsx
--- a/Engine/Source/Programs/Horde/HordeDashboard/plugins/test-automation/index.tsx
+++ b/Engine/Source/Programs/Horde/HordeDashboard/plugins/test-automation/index.tsx
@@ -1,11 +1,15 @@
 // Copyright Epic Games, Inc. All Rights Reserved.
 
 import { registerHordePlugin, MountType } from "..";
+import { TestAuditLogView } from "./testAuditLog";
 import { TestAutomationView } from "./testAutomationView";
 
 registerHordePlugin({
    id: "build",
-   routes: [{ path: "test-automation", element: <TestAutomationView /> }],
+   routes: [
+      { path: "test-automation", element: <TestAutomationView /> },
+      { path: "test-automation/log/:testId", element: <TestAuditLogView /> }
+   ],
    mount: {
       type: MountType.TopNav,
       context: "Tools",
diff --git a/Engine/Source/Programs/Horde/HordeDashboard/plugins/test-automation/testAuditLog.tsx b/Engine/Source/Programs/Horde/HordeDashboard/plugins/test-automation/testAuditLog.tsx
new file mode 100644
--- /dev/null
+++ b/Engine/Source/Programs/Horde/HordeDashboard/plugins/test-automation/testAuditLog.tsx
@@ -0,0 +1,51 @@
+// Copyright Epic Games, Inc. All Rights Reserved.
+import { useParams } from "react-router-dom";
+import { useEffect, useMemo, useState } from "react";
+import { Stack } from "@fluentui/react";
+import { AuditLine, AuditLogPanel, ItemAuditLogConfig } from "horde/components/AuditLog";
+import { TopNav } from "horde/components/TopNav";
+import { getHordeStyling } from "horde/styles/Styles";
+import { getTestAuditLogHistoryV2, getTestNamesV2 } from "./api";
+
+class TestAuditLogLine extends AuditLine {
+
+}
+
+export const TestAuditLogView: React.FC = () => {
+
+   const { testId } = useParams<{ testId: string }>();
+   const [testName, setTestName] = useState<string | undefined>(undefined);
+   
+   useEffect(() => {
+      if (testId) {
+         getTestNamesV2([testId]).then((response) => 
+            {
+               if (response.length) {
+                  setTestName(response[0].name);
+               } else {
+                  setTestName(testId);
+               }
+            });
+      }
+   }, [testId]);
+
+   const { hordeClasses } = getHordeStyling();
+
+   const itemAudit: ItemAuditLogConfig | undefined = useMemo(() => !testId && !testName? undefined : {
+      id: testId!,
+      lineFactory: (entry, search) => new TestAuditLogLine(entry, search),
+      getHistory: (testId, minTime, maxTime) => getTestAuditLogHistoryV2(testId, { minTime: minTime?.toISOString(), maxTime: maxTime?.toISOString(), count: 8192 * 2 }),
+      breadCrumbs: [
+         {text: `Test Health`},
+         {text: `Test Audit Log - ${testName}`}
+      ]
+   }, [testId, testName]);
+
+   return <Stack className={hordeClasses.horde}>
+            <TopNav />
+            <Stack>
+               {!!itemAudit && <AuditLogPanel itemAudit={itemAudit} />}
+            </Stack>
+         </Stack>
+
+}
diff --git a/Engine/Source/Programs/Horde/HordeDashboard/plugins/test-automation/testAutomationHealth.tsx b/Engine/Source/Programs/Horde/HordeDashboard/plugins/test-automation/testAutomationHealth.tsx
--- a/Engine/Source/Programs/Horde/HordeDashboard/plugins/test-automation/testAutomationHealth.tsx
+++ b/Engine/Source/Programs/Horde/HordeDashboard/plugins/test-automation/testAutomationHealth.tsx
@@ -1,6 +1,6 @@
 // Copyright Epic Games, Inc. All Rights Reserved.
 
-import { DetailsList, FontIcon, IColumn, IconButton, Label, Modal, SelectionMode, Stack, Text, HoverCard, HoverCardType, IPlainCardProps, DirectionalHint, DetailsListLayoutMode, Spinner, SpinnerSize, Toggle, DatePicker, ComboBox, IComboBoxOption, FocusZone, FocusZoneTabbableElements, FocusTrapCallout, PrimaryButton, DefaultButton, TooltipHost, TooltipOverflowMode, SuggestionItemType, TooltipDelay } from "@fluentui/react";
+import { DetailsList, FontIcon, IColumn, IconButton, Label, Modal, SelectionMode, Stack, Text, HoverCard, HoverCardType, IPlainCardProps, DirectionalHint, DetailsListLayoutMode, Spinner, SpinnerSize, Toggle, DatePicker, ComboBox, IComboBoxOption, FocusZone, FocusZoneTabbableElements, FocusTrapCallout, PrimaryButton, DefaultButton, TooltipHost, TooltipOverflowMode, TooltipDelay, Link } from "@fluentui/react";
 import { MetadataRef, TestDataHandler, TestNameRef, TestSessionResult } from "./testData";
 import { useCallback, useEffect, useMemo, useState } from "react";
 import { observer } from "mobx-react-lite";
@@ -65,8 +65,10 @@ const renderTestTitle = (item: TestHealthItem): JSX.Element => {
             </HoverCard>
 }
 
-const onRenderExplainCard = (explainData: TestHealthExplain): JSX.Element => {
+const onRenderExplainCard = (item: TestHealthItem): JSX.Element => {
+    const explainData: TestHealthExplain = item.healthExplain!;
     return <Stack style={{pointerEvents: 'none', padding: 10}} tokens={{childrenGap: 5}}>
+                {!!item.underAudit && <Stack style={{ padding: 4, borderRadius: 4, backgroundColor: booleanColors.get(true) }}><Text>Under Audit</Text></Stack>}
                 {!!explainData.successRate && <Text>{explainData.successRate}% Success</Text>}
                 {!!explainData.failureRate &&
                     <Stack>
@@ -81,12 +83,12 @@ const onRenderExplainCard = (explainData: TestHealthExplain): JSX.Element => {
 const renderHealthBadge = (item: TestHealthItem, width?: number): JSX.Element => {
     const StatusCardProps: IPlainCardProps | undefined = item.healthExplain? {
         onRenderPlainCard: onRenderExplainCard,
-        renderData: item.healthExplain,
+        renderData: item,
         directionalHint: DirectionalHint.rightCenter,
         calloutProps: { isBeakVisible: true }
     } : { onRenderPlainCard: () => null /** empty card */ };
     return <HoverCard plainCardProps={StatusCardProps} type={HoverCardType.plain} cardOpenDelay={20} sticky={false}>
-                <Stack style={{backgroundColor: healthBadgeBackground(), padding: 3, borderRadius: 4, width: width, height: 24, cursor: 'inherited'}} horizontal verticalAlign="center" title={item.health} tokens={{childrenGap: 3}} horizontalAlign="center">
+                <Stack style={{backgroundColor: item.underAudit? booleanColors.get(true) : healthBadgeBackground(), padding: 3, borderRadius: 4, width: width, height: 24, cursor: 'inherited'}} horizontal verticalAlign="center" title={item.health} tokens={{childrenGap: 3}} horizontalAlign="center">
                     {new Array(healthStars.get(item.health)).fill(1).map((_, i) => <FontIcon key={`star-${i}`} style={{ fontSize: 11, color: healthColors.get(item.health) }} iconName="Star" />)}
                 </Stack>
             </HoverCard>
@@ -537,6 +539,9 @@ export const TestHealthView: React.FC<{ test: TestHealthItem, handler: TestDataH
 
     useEffect(() => {
         let userIds: string[] = [];
+        if (test.user) {
+            userIds.push(test.user);
+        }
         if (test.owner) {
             userIds.push(test.owner);
         }
@@ -739,36 +744,55 @@ export const TestHealthView: React.FC<{ test: TestHealthItem, handler: TestDataH
                             </Stack>
                         </Stack>
 
-                        <Stack style={{padding: '4px 12px', backgroundColor: booleanColors.get(test.underAudit)}}>
-                            <Stack horizontal tokens={{childrenGap: 6}} verticalAlign="center">
-                                <Toggle
-                                    key={`under-audit-${test.updated}`}
-                                    label="Under Audit"
-                                    checked={test.underAudit}
-                                    onChange={(_, checked) => {test.underAudit = !!checked}}
-                                    inlineLabel
-                                    onText="Yes" offText="No"
-                                />
-                                <Stack horizontal horizontalAlign="start" tokens={{childrenGap: 6}} grow>
-                                    <Label>Last date: </Label>
-                                    <DatePicker
-                                        key={`last-audit-${test.updated}`}
-                                        styles={{root: { width: 150 }}}
-                                        placeholder="Select a date... "
-                                        value={test.lastAuditDate}
-                                        onSelectDate={(date) => test.lastAuditDate = (date ?? undefined)}
+                        <Stack tokens={{childrenGap: 4}}>
+                            <Stack style={{padding: '4px 12px', backgroundColor: booleanColors.get(test.underAudit)}}>
+                                <Stack horizontal tokens={{childrenGap: 6}} verticalAlign="center">
+                                    <Toggle
+                                        key={`under-audit-${test.updated}`}
+                                        label="Under Audit"
+                                        checked={test.underAudit}
+                                        onChange={(_, checked) => {test.underAudit = !!checked}}
+                                        inlineLabel
+                                        onText="Yes" offText="No"
                                     />
+                                    <Stack horizontal horizontalAlign="start" tokens={{childrenGap: 6}}>
+                                        <Label>Last date: </Label>
+                                        <DatePicker
+                                            key={`last-audit-${test.updated}`}
+                                            styles={{root: { width: 150 }}}
+                                            placeholder="Select a date... "
+                                            value={test.lastAuditDate}
+                                            onSelectDate={(date) => test.lastAuditDate = (date ?? undefined)}
+                                        />
+                                    </Stack>
                                 </Stack>
+                                <EditableTextField
+                                    key={`notes-${test.updated}`}
+                                    label="Notes"
+                                    text={test.notes}
+                                    onChange={(value) => test.notes = value}
+                                    multiline
+                                    placeholder="add notes"
+                                    loaded={handler.queryAuditsLoading}
+                                />
+                            </Stack>
+                            
+                            <Stack horizontal horizontalAlign="end" verticalAlign="center" tokens={{childrenGap: 6}} style={{cursor: 'default'}} grow>
+                                {!!test.user &&
+                                    <Stack horizontal verticalAlign="center" tokens={{childrenGap: 6}} >
+                                        <Label style={{fontSize: 12}}>Last update by: </Label>
+                                        <TooltipHost content={`on ${getShortNiceTime(test.userInputDate)}`} delay={TooltipDelay.zero}>
+                                            <Stack style={{...healthStyles.userTag, padding: '1px 7px'}} verticalAlign="center" horizontal>
+                                                <Text variant="small" style={{color:textEmphasisColor(), whiteSpace: 'nowrap'}}>{usersInfoCache.getInfo(test.user)?.name ?? <Spinner size={SpinnerSize.xSmall} />}</Text>
+                                            </Stack>
+                                        </TooltipHost>
+                                    </Stack>
+                                }
+                                {!test.user && <Text style={{color: textEmphasisColor()}} variant="small">no user update</Text>}
+                                <Link href={`/test-automation/log/${test.nameRef.id}`} target="_blank">
+                                    <FontIcon iconName="TimeEntry" style={{cursor: 'pointer', fontSize: 17, padding: 2}} title="open audit log"/>
+                                </Link>
                             </Stack>
-                            <EditableTextField
-                                key={`notes-${test.updated}`}
-                                label="Notes"
-                                text={test.notes}
-                                onChange={(value) => test.notes = value}
-                                multiline
-                                placeholder="add notes"
-                                loaded={handler.queryAuditsLoading}
-                            />
                         </Stack>
 
                     </Stack>
diff --git a/Engine/Source/Programs/Horde/HordeDashboard/plugins/test-automation/testHealthData.ts b/Engine/Source/Programs/Horde/HordeDashboard/plugins/test-automation/testHealthData.ts
--- a/Engine/Source/Programs/Horde/HordeDashboard/plugins/test-automation/testHealthData.ts
+++ b/Engine/Source/Programs/Horde/HordeDashboard/plugins/test-automation/testHealthData.ts
@@ -263,6 +263,14 @@ export class TestHealthItem {
         this.setUpdated();
     }
 
+    get user() {
+        return this.audit?.userId;
+    }
+
+    get userInputDate() {
+        return this.audit?.userInputDate;
+    }
+
     get owner() {
         return this.audit?.owner;
     }
diff --git a/Engine/Source/Programs/Horde/HordeDashboard/src/components/AuditLog.tsx b/Engine/Source/Programs/Horde/HordeDashboard/src/components/AuditLog.tsx
--- a/Engine/Source/Programs/Horde/HordeDashboard/src/components/AuditLog.tsx
+++ b/Engine/Source/Programs/Horde/HordeDashboard/src/components/AuditLog.tsx
@@ -80,6 +80,7 @@ class AuditLogHandler {
       this.deviceId = undefined;
       this.templateAudit = undefined;
       this.jobId = undefined;
+      this.itemAudit = undefined;
 
       this.timeSelectKey = "time_live_tail";
 
@@ -178,6 +179,21 @@ class AuditLogHandler {
       this.update();
    }
 
+   setItem(itemAudit?: ItemAuditLogConfig) {
+
+      if (!itemAudit) {
+         this.clear();
+         return;
+      }
+
+      if (this.itemAudit?.id === itemAudit.id) {
+         return;
+      }
+
+      this.itemAudit = itemAudit;
+      this.update();
+   }
+
    get isAgentLog(): boolean {
       return !!this.agentId;
    }
@@ -198,6 +214,10 @@ class AuditLogHandler {
       return !!this.templateAudit;
    }
 
+   get isGenericItemLog(): boolean {
+      return !!this.itemAudit;
+   }
+
    async update() {
 
       clearTimeout(this.timeoutId);
@@ -217,7 +237,7 @@ class AuditLogHandler {
          this.canceled.add(i);
       }
 
-      if (!this.agentId && !this.issueId && !this.templateAudit && !this.deviceId && !this.jobId) {
+      if (!this.agentId && !this.issueId && !this.templateAudit && !this.deviceId && !this.jobId && !this.itemAudit) {
          return;
       }
 
@@ -270,6 +290,10 @@ class AuditLogHandler {
             entries = await backend.getJobHistory(this.jobId, { minTime: minTime?.toISOString(), maxTime: maxTime?.toISOString(), count: 8192 * 2 });
          }
 
+         if (this.itemAudit) {
+            entries = await this.itemAudit.getHistory(this.itemAudit.id, minTime, maxTime);
+         }
+
          // check for canceled during graph request
          if (this.canceled.has(cancelId)) {
             return;
@@ -391,6 +415,7 @@ class AuditLogHandler {
    deviceId?: string;
    jobId?: string;
    templateAudit?: TemplateAudit;
+   itemAudit?: ItemAuditLogConfig;
 
    timeSelectKey?: string;
 
@@ -466,7 +491,7 @@ const timeSelections: TimeSelection[] = [
 
 const selection = new Selection({ selectionMode: SelectionMode.multiple });
 
-const AuditLogPanel: React.FC<{ agentId?: string, deviceId?: string, jobId?: string, issueId?: string, templateAudit?: TemplateAudit }> = observer(({ agentId, deviceId, jobId, issueId, templateAudit }) => {
+export const AuditLogPanel: React.FC<{ agentId?: string, deviceId?: string, jobId?: string, issueId?: string, templateAudit?: TemplateAudit, itemAudit?: ItemAuditLogConfig }> = observer(({ agentId, deviceId, jobId, issueId, templateAudit, itemAudit }) => {
 
    const windowSize = useWindowSize();
    const [showDatePicker, setShowDatePicker] = useState(false);
@@ -512,7 +537,11 @@ const AuditLogPanel: React.FC<{ agentId?: string, deviceId?: string, jobId?: str
       handler.setTemplateAudit(templateAudit);
    }
 
-   if (!agentId && !deviceId && !issueId && !templateAudit && !jobId) {
+   if (itemAudit) {
+      handler.setItem(itemAudit);
+   }
+
+   if (!agentId && !deviceId && !issueId && !templateAudit && !jobId && !itemAudit) {
       return null;
    }
 
@@ -662,6 +691,10 @@ const AuditLogPanel: React.FC<{ agentId?: string, deviceId?: string, jobId?: str
       ];
    }
 
+   if (itemAudit) {
+      crumbItems = itemAudit.breadCrumbs;
+   }
+
    const onRenderCell = (item?: AuditLogItem, index?: number, isScrolling?: boolean): JSX.Element => {
       const entry = item!.entry;
 
@@ -681,7 +714,7 @@ const AuditLogPanel: React.FC<{ agentId?: string, deviceId?: string, jobId?: str
                   <Stack horizontal disableShrink={true} >
                      <Stack styles={{ root: { width: 140, whiteSpace: "nowrap", fontSize: logMetricNormal.fontSize, userSelect: "none" } }}> {timestamp}</Stack>
                      <Stack styles={{ root: { color: "#8a8a8a", width: 92, paddingRight: 8, whiteSpace: "nowrap", textAlign: "right", fontSize: logMetricNormal.fontSize, userSelect: "none" } }}> [{item!.type}]</Stack>
-                     <div className={auditStyleNormal.logLineOuter}> <Stack styles={{ root: { paddingLeft: 8, paddingRight: 8 } }}> {renderAuditEntry(entry, search)}</Stack></div>
+                     <div className={auditStyleNormal.logLineOuter}> <Stack styles={{ root: { paddingLeft: 8, paddingRight: 8 } }}> {renderAuditEntry(entry, search, itemAudit?.lineFactory)}</Stack></div>
                   </Stack>
                </Stack>
             </div>
@@ -856,62 +889,27 @@ export const AuditLogView: React.FC = () => {
 // some tags depend on other tags to generate the route includig query parameters, and need to have knowledge of the routes themselves,
 // which means can be somewhat data/tag driven, though not entirely
 
-// @todo: We probably want to be able to register subclass constructors with AuditLine, to dynamically add... or use a generic.  
-enum AuditLogType {
-   Unknown,
-   Agent,
-   Device,
-   Job,
-   Issue,
-   Template
-}
-
 type AuditProperty = {
    tag: string; // without enclosing {}
    value: string | number;
    type?: string;
 }
 
-class AuditLine {
+export type AuditLineFactory = (entry: AuditLogEntry, search?: string) => AuditLine;
+export type GetItemHistory = (itemId: string, minTime?: Date, maxTime?: Date) => Promise<AuditLogEntry[]>;
+export type ItemAuditLogConfig = {
+   id: string;
+   lineFactory: AuditLineFactory;
+   getHistory: GetItemHistory;
+   breadCrumbs: BreadcrumbItem[];
+}
 
-   protected constructor(logType: AuditLogType, entry: AuditLogEntry) {
+export class AuditLine {
 
-      this.logType = logType;
+   constructor(entry: AuditLogEntry, search?: string) {
       this.entry = entry;
-
-   }
-
-   static create(logType: AuditLogType, entry: AuditLogEntry, search?: string) {
-
-      let instance: AuditLine | undefined;
-
-      if (logType === AuditLogType.Agent) {
-         instance = new AgentAuditLine(logType, entry);
-      }
-
-      if (logType === AuditLogType.Issue) {
-         instance = new IssueAuditLine(logType, entry);
-      }
-
-      if (logType === AuditLogType.Template) {
-         instance = new TemplateAuditLine(logType, entry);
-      }
-
-      if (logType === AuditLogType.Device) {
-         instance = new DeviceAuditLine(logType, entry);
-      }
-
-      if (logType === AuditLogType.Job) {
-         instance = new JobAuditLine(logType, entry);
-      }
-
-      if (instance) {
-         instance.search = search;
-         instance.process();
-      }
-
-      return instance;
-
+      this.search = search;
+      this.process();
    }
 
    private process() {
@@ -997,7 +995,7 @@ class AuditLine {
          const property = this.properties.get(tag);
 
          if (!property) {
-            return <span>{this.highlight(key, "!==> Missing Property <==!")}</span>;
+            return <span key={key}>{this.highlight(key, "!==> Missing Property <==!")}</span>;
          }
 
          // Common recognizers
@@ -1014,7 +1012,7 @@ class AuditLine {
 
             if (logId && leaseId) {
                let to = `/log/${logId}`;
-               element = <Link to={to} >{this.highlight(key, leaseId)}</Link>
+               element = <Link to={to} key={key}>{this.highlight(key, leaseId)}</Link>
             }
 
          }
@@ -1026,7 +1024,7 @@ class AuditLine {
 
             if (logId) {
                let to = `/log/${logId}`;
-               element = <Link to={to} >{this.highlight(key, logId)}</Link>
+               element = <Link to={to} key={key}>{this.highlight(key, logId)}</Link>
             }
 
          }
@@ -1039,7 +1037,7 @@ class AuditLine {
 
             if (jobId) {
                let to = `/job/${jobId}`;
-               element = <Link to={to} >{this.highlight(key, jobId)}</Link>
+               element = <Link to={to} key={key}>{this.highlight(key, jobId)}</Link>
             }
 
          }
@@ -1052,7 +1050,7 @@ class AuditLine {
 
             if (stepId && jobId) {
                let to = `/job/${jobId}?step=${stepId}`;
-               element = <Link to={to} >{this.highlight(key, stepId)}</Link>
+               element = <Link to={to} key={key}>{this.highlight(key, stepId)}</Link>
             }
 
          }
@@ -1065,7 +1063,7 @@ class AuditLine {
 
             if (batchId && jobId) {
                let to = `/job/${jobId}?batch=${batchId}`;
-               element = <Link to={to} >{this.highlight(key, batchId)}</Link>
+               element = <Link to={to} key={key}>{this.highlight(key, batchId)}</Link>
             }
 
          }
@@ -1077,7 +1075,7 @@ class AuditLine {
 
             if (change) {
                let to = `${dashboard.swarmUrl}/changes/${change}`;
-               element = <a href={to} rel="noreferrer" target="_blank" >{this.highlight(key, change)}</a>
+               element = <a href={to} rel="noreferrer" target="_blank" key={key}>{this.highlight(key, change)}</a>
             }
 
          }
@@ -1086,7 +1084,7 @@ class AuditLine {
          if (tag === "Message") {
             const message = property.value?.toString();
             if (message && message.includes("\n")) {
-               element = <pre style={{margin: 0}}>{this.highlight(key, message)}</pre>;
+               element = <pre style={{margin: 0}} key={key}>{this.highlight(key, message)}</pre>;
             }
          }
 
@@ -1094,7 +1092,7 @@ class AuditLine {
             return element;
          }
 
-         return <span>{this.highlight(key, property.value?.toString())}</span>;
+         return <span key={key}>{this.highlight(key, property.value?.toString())}</span>;
 
 
       });
@@ -1131,7 +1129,7 @@ class AuditLine {
       }).flat();
 
       if (remaining) {
-         const key = `log_line_remaining_remaining_fragment`;
+         const key = `log_line_remaining_fragment`;
          renderedTags.push(<Highlight key={key} search={this.search ? this.search : ""} >{remaining}</Highlight>)
       }
 
@@ -1145,8 +1143,6 @@ class AuditLine {
 
    properties: Map<string, AuditProperty> = new Map();
 
-   logType = AuditLogType.Unknown;
-
    search?: string;
 
    private static tagRegex = /{[^{}]+}/g
@@ -1175,39 +1171,32 @@ class JobAuditLine extends AuditLine {
 // Line Rendering -------------------------------------------------------------------------------------------------------------------
 
 
-const renderAuditEntry = (entry: AuditLogEntry, search?: string) => {
-
-   let type = AuditLogType.Unknown;
+const renderAuditEntry = (entry: AuditLogEntry, search?: string, factory?: AuditLineFactory) => {
 
    if (handler.isAgentLog) {
-      type = AuditLogType.Agent;
+      factory = (entry: AuditLogEntry, search?: string) => new AgentAuditLine(entry, search);
    }
 
    if (handler.isIssueLog) {
-      type = AuditLogType.Issue;
+     factory = (entry: AuditLogEntry, search?: string) => new IssueAuditLine(entry, search);
    }
 
    if (handler.isTemplateAuditLog) {
-      type = AuditLogType.Template;
+      factory = (entry: AuditLogEntry, search?: string) => new TemplateAuditLine(entry, search);
    }
 
    if (handler.isDeviceLog) {
-      type = AuditLogType.Device;
+      factory = (entry: AuditLogEntry, search?: string) => new DeviceAuditLine(entry, search);
    }
 
    if (handler.isJobLog) {
-      type = AuditLogType.Job;
+      factory = (entry: AuditLogEntry, search?: string) => new JobAuditLine(entry, search);
    }
 
-   if (type === AuditLogType.Unknown) {
+   if (!factory) {
       return <div>Unknown Log Type</div>
    }
-
-   const line = AuditLine.create(type, entry, search);
-   if (!line) {
-      return null;
-   }
-
+   const line = factory(entry, search);
    return line.renderLine();
 }
 
