[Horde][Server] Temporarily removing constraint of unused key in order to update sparse.

#Synced-CL 49901343

diff --git a/Engine/Source/Programs/Horde/Plugins/Build/HordeServer.Build/Jobs/JobStepRefCollection.cs b/Engine/Source/Programs/Horde/Plugins/Build/HordeServer.Build/Jobs/JobStepRefCollection.cs
--- a/Engine/Source/Programs/Horde/Plugins/Build/HordeServer.Build/Jobs/JobStepRefCollection.cs
+++ b/Engine/Source/Programs/Horde/Plugins/Build/HordeServer.Build/Jobs/JobStepRefCollection.cs
@@ -164,7 +164,8 @@ namespace HordeServer.Jobs
 			// Using this index, we can cost-effectively query:
 			//	- All job step refs belonging to a specific job
 			//	- All job step refs belonging to a batch (of a job)
-			indexes.Add(keys => keys.Ascending(x => x.JobId).Ascending(x => x.BatchId).Ascending(x => x.StepId));
+			// Note: This is being temporarily disabled due to Mongo not supporting index updates, and needing to clean up the prior sparse
+			//indexes.Add(keys => keys.Ascending(x => x.JobId).Ascending(x => x.BatchId).Ascending(x => x.StepId));
 
 			_jobStepRefs = mongoService.GetCollection<JobStepRef>("JobStepRefs", indexes);
 			_telemetryWriter = telemetryWriter;
