Horde: Fix warning about duplicate DataRow attribute

#rnx

diff --git a/Engine/Source/Programs/Horde/Drivers/JobDriver.Tests/PerforceMaterializerTests.cs b/Engine/Source/Programs/Horde/Drivers/JobDriver.Tests/PerforceMaterializerTests.cs
--- a/Engine/Source/Programs/Horde/Drivers/JobDriver.Tests/PerforceMaterializerTests.cs
+++ b/Engine/Source/Programs/Horde/Drivers/JobDriver.Tests/PerforceMaterializerTests.cs
@@ -245,7 +245,6 @@ public class PerforceMaterializerTests : BasePerforceFixtureTest
 	[DataRow("/Engine/Build/...", "//Foo/Main/Engine/Build/... //MyComputer/Engine/Build/...")]
 	[DataRow("-/Engine/Docs/...", "-//Foo/Main/Engine/Docs/... //MyComputer/Engine/Docs/...")]
 	[DataRow("   -/Engine/Docs/...", "-//Foo/Main/Engine/Docs/... //MyComputer/Engine/Docs/...")]
-	[DataRow("/Engine/Build/...", "//Foo/Main/Engine/Build/... //MyComputer/Engine/Build/...")]
 	public void ConvertHordeViews(string hordeView, string expected)
 	{
 		string actual = PerforceMaterializer.ConvertHordeViewToClientView(hordeView, "//Foo/Main", "MyComputer");
