Horde: Expose LeaseOutcomeReason in GetJobBatchResponse

Only fetches lease data when batch has an error to avoid unnecessary DB queries for successful batches.

diff --git a/Engine/Source/Programs/Horde/Plugins/Build/HordeServer.Build/Jobs/JobsController.cs b/Engine/Source/Programs/Horde/Plugins/Build/HordeServer.Build/Jobs/JobsController.cs
--- a/Engine/Source/Programs/Horde/Plugins/Build/HordeServer.Build/Jobs/JobsController.cs
+++ b/Engine/Source/Programs/Horde/Plugins/Build/HordeServer.Build/Jobs/JobsController.cs
@@ -1,6 +1,7 @@
 // Copyright Epic Games, Inc. All Rights Reserved.
 
 using EpicGames.Core;
+using EpicGames.Horde.Agents.Leases;
 using EpicGames.Horde.Artifacts;
 using EpicGames.Horde.Commits;
 using EpicGames.Horde.Jobs;
@@ -14,6 +15,7 @@ using EpicGames.Horde.Users;
 using HordeServer.Auditing;
 using HordeServer.Acls;
 using HordeServer.Agents;
+using HordeServer.Agents.Leases;
 using HordeServer.Commits;
 using HordeServer.Jobs.Graphs;
 using HordeServer.Jobs.Templates;
@@ -48,13 +50,14 @@ namespace HordeServer.Jobs
 		private readonly IUserCollection _userCollection;
 		private readonly INotificationService _notificationService;
 		private readonly AgentService _agentService;
+		private readonly ILeaseCollection _leaseCollection;
 		private readonly IOptionsSnapshot<BuildConfig> _buildConfig;
 		private readonly ILogger<JobsController> _logger;
 
 		/// <summary>
 		/// Constructor
 		/// </summary>
-		public JobsController(IGraphCollection graphs, ICommitService commitService, IPerforceService perforce, JobService jobService, ITemplateCollectionInternal templateCollection, IArtifactCollection artifactCollection, IUserCollection userCollection, INotificationService notificationService, AgentService agentService, IOptionsSnapshot<BuildConfig> buildConfig, ILogger<JobsController> logger)
+		public JobsController(IGraphCollection graphs, ICommitService commitService, IPerforceService perforce, JobService jobService, ITemplateCollectionInternal templateCollection, IArtifactCollection artifactCollection, IUserCollection userCollection, INotificationService notificationService, AgentService agentService, ILeaseCollection leaseCollection, IOptionsSnapshot<BuildConfig> buildConfig, ILogger<JobsController> logger)
 		{
 			_graphs = graphs;
 			_commitService = commitService;
@@ -65,6 +68,7 @@ namespace HordeServer.Jobs
 			_userCollection = userCollection;
 			_notificationService = notificationService;
 			_agentService = agentService;
+			_leaseCollection = leaseCollection;
 			_buildConfig = buildConfig;
 			_logger = logger;
 		}
@@ -767,7 +771,18 @@ namespace HordeServer.Jobs
 				agentRate = await _agentService.GetRateAsync(batch.AgentId.Value, cancellationToken);
 			}
 
-			return CreateGetBatchResponse(batch, steps, agentRate);
+			// Only fetch lease for outcome reasons when batch has an error
+			List<string>? leaseOutcomeReasons = null;
+			if (batch.Error != JobStepBatchError.None && batch.LeaseId != null)
+			{
+				ILease? lease = await _leaseCollection.GetAsync(batch.LeaseId.Value, cancellationToken);
+				if (lease != null)
+				{
+					leaseOutcomeReasons = ToFlagNames(lease.OutcomeReason);
+				}
+			}
+
+			return CreateGetBatchResponse(batch, steps, agentRate, leaseOutcomeReasons);
 		}
 
 		/// <summary>
@@ -776,8 +791,9 @@ namespace HordeServer.Jobs
 		/// <param name="batch">The batch to construct from</param>
 		/// <param name="steps">Steps in this batch</param>
 		/// <param name="agentRate">Rate for this agent</param>
+		/// <param name="leaseOutcomeReasons">Lease outcome reasons (only when batch has an error)</param>
 		/// <returns>Response instance</returns>
-		static GetJobBatchResponse CreateGetBatchResponse(IJobStepBatch batch, List<GetJobStepResponse> steps, double? agentRate)
+		static GetJobBatchResponse CreateGetBatchResponse(IJobStepBatch batch, List<GetJobStepResponse> steps, double? agentRate, List<string>? leaseOutcomeReasons)
 		{
 			GetJobBatchResponse response = new GetJobBatchResponse();
 			response.Id = batch.Id;
@@ -791,6 +807,7 @@ namespace HordeServer.Jobs
 			response.AgentRate = agentRate;
 			response.SessionId = batch.SessionId?.ToString();
 			response.LeaseId = batch.LeaseId?.ToString();
+			response.LeaseOutcomeReasons = leaseOutcomeReasons;
 			response.WeightedPriority = batch.SchedulePriority;
 			response.StartTime = batch.StartTimeUtc;
 			response.FinishTime = batch.FinishTimeUtc;
@@ -798,6 +815,25 @@ namespace HordeServer.Jobs
 			return response;
 		}
 
+		/// <summary>
+		/// Converts a LeaseOutcomeReason flags enum to a list of flag names
+		/// </summary>
+		/// <param name="reason">The flags enum value</param>
+		/// <returns>List of flag names, or null if none set</returns>
+		static List<string>? ToFlagNames(LeaseOutcomeReason reason)
+		{
+			if (reason == LeaseOutcomeReason.None)
+			{
+				return null;
+			}
+
+			return Enum.GetValues(typeof(LeaseOutcomeReason))
+				.Cast<LeaseOutcomeReason>()
+				.Where(flag => flag != LeaseOutcomeReason.None && reason.HasFlag(flag))
+				.Select(flag => flag.ToString())
+				.ToList();
+		}
+
 		/// <summary>
 		/// Get the response object for a step
 		/// </summary>
diff --git a/Engine/Source/Programs/Shared/EpicGames.Horde/Jobs/JobMessages.cs b/Engine/Source/Programs/Shared/EpicGames.Horde/Jobs/JobMessages.cs
--- a/Engine/Source/Programs/Shared/EpicGames.Horde/Jobs/JobMessages.cs
+++ b/Engine/Source/Programs/Shared/EpicGames.Horde/Jobs/JobMessages.cs
@@ -1284,6 +1284,11 @@ namespace EpicGames.Horde.Jobs
 		/// </summary>
 		public DateTimeOffset? ReadyTime { get; set; }
 
+		/// <summary>
+		/// Human-readable names of lease outcome reasons (only populated when batch has an error)
+		/// </summary>
+		public List<string>? LeaseOutcomeReasons { get; set; }
+
 		/// <summary>
 		/// Steps within this run
 		/// </summary>
