Horde: Log reason for job driver performing cleanup

#skipci

diff --git a/Engine/Source/Programs/Horde/HordeAgent/Leases/Handlers/JobHandler.cs b/Engine/Source/Programs/Horde/HordeAgent/Leases/Handlers/JobHandler.cs
--- a/Engine/Source/Programs/Horde/HordeAgent/Leases/Handlers/JobHandler.cs
+++ b/Engine/Source/Programs/Horde/HordeAgent/Leases/Handlers/JobHandler.cs
@@ -114,7 +114,7 @@ namespace HordeAgent.Leases.Handlers
 			arguments.Add($"-LeaseId={leaseId}");
 			arguments.Add($"-WorkingDir={session.WorkingDir}");
 
-			string driverName = String.IsNullOrEmpty(executeTask.JobOptions.Driver) ? "JobDriver" : executeTask.JobOptions.Driver;			
+			string driverName = String.IsNullOrEmpty(executeTask.JobOptions.Driver) ? "JobDriver" : executeTask.JobOptions.Driver;
 			FileReference driverAssembly = FileReference.Combine(new DirectoryReference(AppContext.BaseDirectory), driverName, $"{driverName}.dll");			
 			span.SetAttribute("horde.job.driver_name", driverName);
 
@@ -166,7 +166,7 @@ namespace HordeAgent.Leases.Handlers
 			{
 				try
 				{
-					logger.LogInformation("Driver initiating job batch cleanup");
+					logger.LogInformation("Unexpected exit of job driver. Initiating job batch cleanup... (exitCode={ExitCode} exception={Exception})", exitCode, exception?.Message);
 					// Run job batch cleanup
 					environment["UE_HORDE_JOB_BATCH_CLEANUP"] = "1";
 					int cleanupCode = await RunDotNetProcessAsync(driverAssembly, arguments, environment, AgentApp.IsSelfContained, logger, CancellationToken.None);
