Remove suppression of CS1701+1702 as they are no longer needed

#rnx
#Synced-CL 49956463

diff --git a/Engine/Source/Programs/Horde/Drivers/JobDriver.Tests/JobDriver.Tests.csproj b/Engine/Source/Programs/Horde/Drivers/JobDriver.Tests/JobDriver.Tests.csproj
--- a/Engine/Source/Programs/Horde/Drivers/JobDriver.Tests/JobDriver.Tests.csproj
+++ b/Engine/Source/Programs/Horde/Drivers/JobDriver.Tests/JobDriver.Tests.csproj
@@ -6,8 +6,7 @@
     <Nullable>enable</Nullable>
     <IsPackable>false</IsPackable>
     <Configurations>Debug;Release;Development;Analyze</Configurations>
-    <NoWarn>$(NoWarn);1701;1702;1591</NoWarn>
-    <GenerateDocumentationFile>false</GenerateDocumentationFile>
+    <NoWarn>$(NoWarn);1591</NoWarn>
   </PropertyGroup>
 
   <PropertyGroup Condition=" '$(Configuration)|$(Platform)' == 'Analyze|AnyCPU' ">
diff --git a/Engine/Source/Programs/Horde/Drivers/JobDriver/JobDriver.csproj b/Engine/Source/Programs/Horde/Drivers/JobDriver/JobDriver.csproj
--- a/Engine/Source/Programs/Horde/Drivers/JobDriver/JobDriver.csproj
+++ b/Engine/Source/Programs/Horde/Drivers/JobDriver/JobDriver.csproj
@@ -1,4 +1,4 @@
-﻿<Project Sdk="Microsoft.NET.Sdk">
+<Project Sdk="Microsoft.NET.Sdk">
   <Import Project="..\..\..\Shared\UnrealEngine.csproj.props" />
 
   <PropertyGroup>
@@ -10,7 +10,7 @@
     <PublishRepositoryUrl>true</PublishRepositoryUrl>
     <EmbedUntrackedSources>false</EmbedUntrackedSources>
     <Configurations>Debug;Release;Analyze</Configurations>
-    <NoWarn>$(NoWarn);1701;1702;NU3003</NoWarn>
+    <NoWarn>$(NoWarn);NU3003</NoWarn>
   </PropertyGroup>
 
   <PropertyGroup Condition="'$(Configuration)|$(Platform)'=='Debug|AnyCPU'">
diff --git a/Engine/Source/Programs/Horde/HordeAgent.Tests/HordeAgent.Tests.csproj b/Engine/Source/Programs/Horde/HordeAgent.Tests/HordeAgent.Tests.csproj
--- a/Engine/Source/Programs/Horde/HordeAgent.Tests/HordeAgent.Tests.csproj
+++ b/Engine/Source/Programs/Horde/HordeAgent.Tests/HordeAgent.Tests.csproj
@@ -6,7 +6,7 @@
     <Nullable>enable</Nullable>
     <IsPackable>false</IsPackable>
     <Configurations>Debug;Release;Development;Analyze</Configurations>
-    <NoWarn>$(NoWarn);1701;1702;1591</NoWarn>
+    <NoWarn>$(NoWarn);1591</NoWarn>
     <RunAnalyzersDuringBuild>False</RunAnalyzersDuringBuild>
   </PropertyGroup>
 
diff --git a/Engine/Source/Programs/Horde/HordeAgent/HordeAgent.csproj b/Engine/Source/Programs/Horde/HordeAgent/HordeAgent.csproj
--- a/Engine/Source/Programs/Horde/HordeAgent/HordeAgent.csproj
+++ b/Engine/Source/Programs/Horde/HordeAgent/HordeAgent.csproj
@@ -12,7 +12,7 @@
     <PublishRepositoryUrl>true</PublishRepositoryUrl>
     <EmbedUntrackedSources>false</EmbedUntrackedSources>
     <Configurations>Debug;Release;Analyze</Configurations>
-    <NoWarn>$(NoWarn);1701;1702;NU3003</NoWarn>
+    <NoWarn>$(NoWarn);NU3003</NoWarn>
     <RunAnalyzersDuringBuild>False</RunAnalyzersDuringBuild>
   </PropertyGroup>
 
@@ -207,7 +207,7 @@
       <FilesToDelete Include="$(PublishDir)**/oo2core*winuwp*.dll" />
     </ItemGroup>
     <Delete Files="@(FilesToDelete)">
-      <Output TaskParameter="DeletedFiles" ItemName="DeletedList"/>
+      <Output TaskParameter="DeletedFiles" ItemName="DeletedList" />
     </Delete>
     <Message Text="Deleted files:%0A@(DeletedList->'  %(Identity)','%0A')" Importance="high" />
   </Target>
diff --git a/Engine/Source/Programs/Horde/HordeServer.Tests.Shared/HordeServer.Tests.Shared.csproj b/Engine/Source/Programs/Horde/HordeServer.Tests.Shared/HordeServer.Tests.Shared.csproj
--- a/Engine/Source/Programs/Horde/HordeServer.Tests.Shared/HordeServer.Tests.Shared.csproj
+++ b/Engine/Source/Programs/Horde/HordeServer.Tests.Shared/HordeServer.Tests.Shared.csproj
@@ -8,7 +8,7 @@
     <RootNamespace>HordeServer.Tests</RootNamespace>
     <RunAnalyzersDuringBuild>False</RunAnalyzersDuringBuild>
     <Configurations>Debug;Release;Analyze</Configurations>
-    <NoWarn>$(NoWarn);1701;1702;1591</NoWarn>
+    <NoWarn>$(NoWarn);1591</NoWarn>
   </PropertyGroup>
 
   <PropertyGroup Condition=" '$(Configuration)|$(Platform)' == 'Analyze|AnyCPU' ">
diff --git a/Engine/Source/Programs/Horde/HordeServer.Tests/HordeServer.Tests.csproj b/Engine/Source/Programs/Horde/HordeServer.Tests/HordeServer.Tests.csproj
--- a/Engine/Source/Programs/Horde/HordeServer.Tests/HordeServer.Tests.csproj
+++ b/Engine/Source/Programs/Horde/HordeServer.Tests/HordeServer.Tests.csproj
@@ -7,7 +7,7 @@
     <IsPackable>false</IsPackable>
     <NeutralLanguage>en-US</NeutralLanguage>
     <Configurations>Debug;Release;Development;Analyze</Configurations>
-    <NoWarn>$(NoWarn);1701;1702;1591</NoWarn>
+    <NoWarn>$(NoWarn);1591</NoWarn>
     <GenerateProgramFile>false</GenerateProgramFile>
     <RunAnalyzersDuringBuild>False</RunAnalyzersDuringBuild>
   </PropertyGroup>
diff --git a/Engine/Source/Programs/Horde/Plugins/Analytics/HordeServer.Analytics.Integration.Tests/HordeServer.Analytics.Integration.Tests.csproj b/Engine/Source/Programs/Horde/Plugins/Analytics/HordeServer.Analytics.Integration.Tests/HordeServer.Analytics.Integration.Tests.csproj
--- a/Engine/Source/Programs/Horde/Plugins/Analytics/HordeServer.Analytics.Integration.Tests/HordeServer.Analytics.Integration.Tests.csproj
+++ b/Engine/Source/Programs/Horde/Plugins/Analytics/HordeServer.Analytics.Integration.Tests/HordeServer.Analytics.Integration.Tests.csproj
@@ -9,7 +9,7 @@
     <IsTestProject>true</IsTestProject>
     <RunAnalyzersDuringBuild>False</RunAnalyzersDuringBuild>
     <Configurations>Debug;Release;Analyze</Configurations>
-    <NoWarn>$(NoWarn);1701;1702;1591</NoWarn>
+    <NoWarn>$(NoWarn);1591</NoWarn>
   </PropertyGroup>
 
   <PropertyGroup Condition=" '$(Configuration)|$(Platform)' == 'Analyze|AnyCPU' ">
diff --git a/Engine/Source/Programs/Horde/Plugins/Analytics/HordeServer.Analytics.Tests/HordeServer.Analytics.Tests.csproj b/Engine/Source/Programs/Horde/Plugins/Analytics/HordeServer.Analytics.Tests/HordeServer.Analytics.Tests.csproj
--- a/Engine/Source/Programs/Horde/Plugins/Analytics/HordeServer.Analytics.Tests/HordeServer.Analytics.Tests.csproj
+++ b/Engine/Source/Programs/Horde/Plugins/Analytics/HordeServer.Analytics.Tests/HordeServer.Analytics.Tests.csproj
@@ -9,7 +9,7 @@
     <IsTestProject>true</IsTestProject>
     <RunAnalyzersDuringBuild>False</RunAnalyzersDuringBuild>
     <Configurations>Debug;Release;Analyze</Configurations>
-    <NoWarn>$(NoWarn);1701;1702;1591</NoWarn>
+    <NoWarn>$(NoWarn);1591</NoWarn>
   </PropertyGroup>
 
   <PropertyGroup Condition=" '$(Configuration)|$(Platform)' == 'Analyze|AnyCPU' ">
diff --git a/Engine/Source/Programs/Horde/Plugins/Build/HordeServer.Build.Tests/HordeServer.Build.Tests.csproj b/Engine/Source/Programs/Horde/Plugins/Build/HordeServer.Build.Tests/HordeServer.Build.Tests.csproj
--- a/Engine/Source/Programs/Horde/Plugins/Build/HordeServer.Build.Tests/HordeServer.Build.Tests.csproj
+++ b/Engine/Source/Programs/Horde/Plugins/Build/HordeServer.Build.Tests/HordeServer.Build.Tests.csproj
@@ -10,7 +10,7 @@
     <RootNamespace>HordeServer.Tests</RootNamespace>
     <RunAnalyzersDuringBuild>False</RunAnalyzersDuringBuild>
     <Configurations>Debug;Release;Analyze</Configurations>
-    <NoWarn>$(NoWarn);1701;1702;1591</NoWarn>
+    <NoWarn>$(NoWarn);1591</NoWarn>
   </PropertyGroup>
 
   <PropertyGroup Condition=" '$(Configuration)|$(Platform)' == 'Analyze|AnyCPU' ">
diff --git a/Engine/Source/Programs/Horde/Plugins/Compute/HordeServer.Compute.Tests/HordeServer.Compute.Tests.csproj b/Engine/Source/Programs/Horde/Plugins/Compute/HordeServer.Compute.Tests/HordeServer.Compute.Tests.csproj
--- a/Engine/Source/Programs/Horde/Plugins/Compute/HordeServer.Compute.Tests/HordeServer.Compute.Tests.csproj
+++ b/Engine/Source/Programs/Horde/Plugins/Compute/HordeServer.Compute.Tests/HordeServer.Compute.Tests.csproj
@@ -10,7 +10,7 @@
     <RootNamespace>HordeServer.Tests</RootNamespace>
     <RunAnalyzersDuringBuild>False</RunAnalyzersDuringBuild>
     <Configurations>Debug;Release;Analyze</Configurations>
-    <NoWarn>$(NoWarn);1701;1702;1591</NoWarn>
+    <NoWarn>$(NoWarn);1591</NoWarn>
   </PropertyGroup>
 
   <PropertyGroup Condition=" '$(Configuration)|$(Platform)' == 'Analyze|AnyCPU' ">
diff --git a/Engine/Source/Programs/Horde/Plugins/ExamplePlugin/HordeServer.ExamplePlugin.Tests/HordeServer.ExamplePlugin.Tests.csproj b/Engine/Source/Programs/Horde/Plugins/ExamplePlugin/HordeServer.ExamplePlugin.Tests/HordeServer.ExamplePlugin.Tests.csproj
--- a/Engine/Source/Programs/Horde/Plugins/ExamplePlugin/HordeServer.ExamplePlugin.Tests/HordeServer.ExamplePlugin.Tests.csproj
+++ b/Engine/Source/Programs/Horde/Plugins/ExamplePlugin/HordeServer.ExamplePlugin.Tests/HordeServer.ExamplePlugin.Tests.csproj
@@ -10,7 +10,7 @@
     <RootNamespace>HordeServer.Tests</RootNamespace>
     <RunAnalyzersDuringBuild>False</RunAnalyzersDuringBuild>
     <Configurations>Debug;Release;Analyze</Configurations>
-    <NoWarn>$(NoWarn);1701;1702;1591</NoWarn>
+    <NoWarn>$(NoWarn);1591</NoWarn>
   </PropertyGroup>
 
   <PropertyGroup Condition=" '$(Configuration)|$(Platform)' == 'Analyze|AnyCPU' ">
diff --git a/Engine/Source/Programs/Horde/Plugins/Experimental/HordeServer.Experimental.Tests/HordeServer.Experimental.Tests.csproj b/Engine/Source/Programs/Horde/Plugins/Experimental/HordeServer.Experimental.Tests/HordeServer.Experimental.Tests.csproj
--- a/Engine/Source/Programs/Horde/Plugins/Experimental/HordeServer.Experimental.Tests/HordeServer.Experimental.Tests.csproj
+++ b/Engine/Source/Programs/Horde/Plugins/Experimental/HordeServer.Experimental.Tests/HordeServer.Experimental.Tests.csproj
@@ -10,7 +10,7 @@
     <RootNamespace>HordeServer.Tests</RootNamespace>
     <RunAnalyzersDuringBuild>False</RunAnalyzersDuringBuild>
     <Configurations>Debug;Release;Analyze</Configurations>
-    <NoWarn>$(NoWarn);1701;1702;1591</NoWarn>
+    <NoWarn>$(NoWarn);1591</NoWarn>
   </PropertyGroup>
 
   <PropertyGroup Condition=" '$(Configuration)|$(Platform)' == 'Analyze|AnyCPU' ">
diff --git a/Engine/Source/Programs/Horde/Plugins/Secrets/HordeServer.Secrets.Tests/HordeServer.Secrets.Tests.csproj b/Engine/Source/Programs/Horde/Plugins/Secrets/HordeServer.Secrets.Tests/HordeServer.Secrets.Tests.csproj
--- a/Engine/Source/Programs/Horde/Plugins/Secrets/HordeServer.Secrets.Tests/HordeServer.Secrets.Tests.csproj
+++ b/Engine/Source/Programs/Horde/Plugins/Secrets/HordeServer.Secrets.Tests/HordeServer.Secrets.Tests.csproj
@@ -10,7 +10,7 @@
     <RootNamespace>HordeServer.Tests</RootNamespace>
     <RunAnalyzersDuringBuild>False</RunAnalyzersDuringBuild>
     <Configurations>Debug;Release;Analyze</Configurations>
-    <NoWarn>$(NoWarn);1701;1702;1591</NoWarn>
+    <NoWarn>$(NoWarn);1591</NoWarn>
   </PropertyGroup>
 
   <PropertyGroup Condition=" '$(Configuration)|$(Platform)' == 'Analyze|AnyCPU' ">
diff --git a/Engine/Source/Programs/Horde/Plugins/Storage/HordeServer.Storage.Tests/HordeServer.Storage.Tests.csproj b/Engine/Source/Programs/Horde/Plugins/Storage/HordeServer.Storage.Tests/HordeServer.Storage.Tests.csproj
--- a/Engine/Source/Programs/Horde/Plugins/Storage/HordeServer.Storage.Tests/HordeServer.Storage.Tests.csproj
+++ b/Engine/Source/Programs/Horde/Plugins/Storage/HordeServer.Storage.Tests/HordeServer.Storage.Tests.csproj
@@ -10,7 +10,7 @@
     <RootNamespace>HordeServer.Tests</RootNamespace>
     <RunAnalyzersDuringBuild>False</RunAnalyzersDuringBuild>
     <Configurations>Debug;Release;Analyze</Configurations>
-    <NoWarn>$(NoWarn);1701;1702;1591</NoWarn>
+    <NoWarn>$(NoWarn);1591</NoWarn>
   </PropertyGroup>
 
   <PropertyGroup Condition=" '$(Configuration)|$(Platform)' == 'Analyze|AnyCPU' ">
diff --git a/Engine/Source/Programs/Horde/Plugins/Tools/HordeServer.Tools.Tests/HordeServer.Tools.Tests.csproj b/Engine/Source/Programs/Horde/Plugins/Tools/HordeServer.Tools.Tests/HordeServer.Tools.Tests.csproj
--- a/Engine/Source/Programs/Horde/Plugins/Tools/HordeServer.Tools.Tests/HordeServer.Tools.Tests.csproj
+++ b/Engine/Source/Programs/Horde/Plugins/Tools/HordeServer.Tools.Tests/HordeServer.Tools.Tests.csproj
@@ -10,7 +10,7 @@
     <RootNamespace>HordeServer.Tests</RootNamespace>
     <RunAnalyzersDuringBuild>False</RunAnalyzersDuringBuild>
     <Configurations>Debug;Release;Analyze</Configurations>
-    <NoWarn>$(NoWarn);1701;1702;1591</NoWarn>
+    <NoWarn>$(NoWarn);1591</NoWarn>
   </PropertyGroup>
 
   <PropertyGroup Condition=" '$(Configuration)|$(Platform)' == 'Analyze|AnyCPU' ">
diff --git a/Engine/Source/Programs/Shared/EpicGames.AspNet/EpicGames.AspNet.csproj b/Engine/Source/Programs/Shared/EpicGames.AspNet/EpicGames.AspNet.csproj
--- a/Engine/Source/Programs/Shared/EpicGames.AspNet/EpicGames.AspNet.csproj
+++ b/Engine/Source/Programs/Shared/EpicGames.AspNet/EpicGames.AspNet.csproj
@@ -4,7 +4,7 @@
     <!-- This needs to be explictly declared as a library again to build. -->
     <OutputType>Library</OutputType>
     <DebugType>pdbonly</DebugType>
-    <NoWarn>$(NoWarn);1701;1702;1591</NoWarn>
+    <NoWarn>$(NoWarn);1591</NoWarn>
   </PropertyGroup>
 
   <ItemGroup>
diff --git a/Engine/Source/Programs/Shared/EpicGames.Build/EpicGames.Build.csproj b/Engine/Source/Programs/Shared/EpicGames.Build/EpicGames.Build.csproj
--- a/Engine/Source/Programs/Shared/EpicGames.Build/EpicGames.Build.csproj
+++ b/Engine/Source/Programs/Shared/EpicGames.Build/EpicGames.Build.csproj
@@ -6,7 +6,7 @@
   </PropertyGroup>
 
   <PropertyGroup Condition=" '$(Configuration)|$(Platform)' == 'Analyze|AnyCPU' ">
-    <NoWarn>$(NoWarn);1701;1702;1591</NoWarn>
+    <NoWarn>$(NoWarn);1591</NoWarn>
   </PropertyGroup>
 
   <ItemGroup>
diff --git a/Engine/Source/Programs/Shared/EpicGames.BuildGraph.Tests/EpicGames.BuildGraph.Tests.csproj b/Engine/Source/Programs/Shared/EpicGames.BuildGraph.Tests/EpicGames.BuildGraph.Tests.csproj
--- a/Engine/Source/Programs/Shared/EpicGames.BuildGraph.Tests/EpicGames.BuildGraph.Tests.csproj
+++ b/Engine/Source/Programs/Shared/EpicGames.BuildGraph.Tests/EpicGames.BuildGraph.Tests.csproj
@@ -1,7 +1,6 @@
 ﻿<Project Sdk="Microsoft.NET.Sdk">
   <PropertyGroup>
     <TargetFramework>net8.0</TargetFramework>
-    <NoWarn>$(NoWarn);1701;1702</NoWarn>
   </PropertyGroup>
 
   <ItemGroup>
diff --git a/Engine/Source/Programs/Shared/EpicGames.Core.Tests/EpicGames.Core.Tests.csproj b/Engine/Source/Programs/Shared/EpicGames.Core.Tests/EpicGames.Core.Tests.csproj
--- a/Engine/Source/Programs/Shared/EpicGames.Core.Tests/EpicGames.Core.Tests.csproj
+++ b/Engine/Source/Programs/Shared/EpicGames.Core.Tests/EpicGames.Core.Tests.csproj
@@ -1,7 +1,6 @@
 <Project Sdk="Microsoft.NET.Sdk">
   <PropertyGroup>
     <TargetFramework>net8.0</TargetFramework>
-    <NoWarn>$(NoWarn);1701;1702</NoWarn>
     <TreatWarningsAsErrors>False</TreatWarningsAsErrors>
   </PropertyGroup>
 
diff --git a/Engine/Source/Programs/Shared/EpicGames.Horde.Tests/EpicGames.Horde.Tests.csproj b/Engine/Source/Programs/Shared/EpicGames.Horde.Tests/EpicGames.Horde.Tests.csproj
--- a/Engine/Source/Programs/Shared/EpicGames.Horde.Tests/EpicGames.Horde.Tests.csproj
+++ b/Engine/Source/Programs/Shared/EpicGames.Horde.Tests/EpicGames.Horde.Tests.csproj
@@ -1,7 +1,6 @@
 <Project Sdk="Microsoft.NET.Sdk">
   <PropertyGroup>
     <TargetFramework>net8.0</TargetFramework>
-    <NoWarn>$(NoWarn);1701;1702</NoWarn>
   </PropertyGroup>
 
   <ItemGroup>
diff --git a/Engine/Source/Programs/Shared/EpicGames.LogMatchers.Tests/EpicGames.LogMatchers.Tests.csproj b/Engine/Source/Programs/Shared/EpicGames.LogMatchers.Tests/EpicGames.LogMatchers.Tests.csproj
--- a/Engine/Source/Programs/Shared/EpicGames.LogMatchers.Tests/EpicGames.LogMatchers.Tests.csproj
+++ b/Engine/Source/Programs/Shared/EpicGames.LogMatchers.Tests/EpicGames.LogMatchers.Tests.csproj
@@ -1,7 +1,6 @@
 <Project Sdk="Microsoft.NET.Sdk">
   <PropertyGroup>
     <TargetFramework>net8.0</TargetFramework>
-    <NoWarn>$(NoWarn);1701;1702</NoWarn>
   </PropertyGroup>
 
   <ItemGroup>
diff --git a/Engine/Source/Programs/Shared/EpicGames.Redis.Tests/EpicGames.Redis.Tests.csproj b/Engine/Source/Programs/Shared/EpicGames.Redis.Tests/EpicGames.Redis.Tests.csproj
--- a/Engine/Source/Programs/Shared/EpicGames.Redis.Tests/EpicGames.Redis.Tests.csproj
+++ b/Engine/Source/Programs/Shared/EpicGames.Redis.Tests/EpicGames.Redis.Tests.csproj
@@ -1,7 +1,6 @@
 <Project Sdk="Microsoft.NET.Sdk">
   <PropertyGroup>
     <TargetFramework>net8.0</TargetFramework>
-    <NoWarn>$(NoWarn);1701;1702</NoWarn>
   </PropertyGroup>
 
   <ItemGroup>
diff --git a/Engine/Source/Programs/Shared/EpicGames.Serialization.Tests/EpicGames.Serialization.Tests.csproj b/Engine/Source/Programs/Shared/EpicGames.Serialization.Tests/EpicGames.Serialization.Tests.csproj
--- a/Engine/Source/Programs/Shared/EpicGames.Serialization.Tests/EpicGames.Serialization.Tests.csproj
+++ b/Engine/Source/Programs/Shared/EpicGames.Serialization.Tests/EpicGames.Serialization.Tests.csproj
@@ -1,7 +1,6 @@
 <Project Sdk="Microsoft.NET.Sdk">
   <PropertyGroup>
     <TargetFramework>net8.0</TargetFramework>
-    <NoWarn>$(NoWarn);1701;1702</NoWarn>
   </PropertyGroup>
 
   <ItemGroup>
