Horde: Add support for change button to not require job object

diff --git a/Engine/Source/Programs/Horde/HordeDashboard/src/components/ChangeButton.tsx b/Engine/Source/Programs/Horde/HordeDashboard/src/components/ChangeButton.tsx
--- a/Engine/Source/Programs/Horde/HordeDashboard/src/components/ChangeButton.tsx
+++ b/Engine/Source/Programs/Horde/HordeDashboard/src/components/ChangeButton.tsx
@@ -26,16 +26,17 @@ export type JobParameters = {
    batches?: GetBatchResponse[];
 }
 
-export const ChangeContextMenu: React.FC<{ target: ChangeContextMenuTarget, onDismiss?: () => void, job?: JobParameters, stepRef?: GetJobStepRefResponse, commit?: GetChangeSummaryResponse, rangeCL?: number }> = ({ target, onDismiss, job, stepRef, commit, rangeCL }) => {
+export const ChangeContextMenu: React.FC<{ target: ChangeContextMenuTarget, onDismiss?: () => void, job?: JobParameters, stepRef?: GetJobStepRefResponse, commit?: GetChangeSummaryResponse, rangeCL?: number, cl?: number }> = ({ target, onDismiss, job, stepRef, commit, rangeCL, cl }) => {
 
-   if (!job) {
-      return null;
-   }
+   let jobChange = commit?.number ?? cl;
 
-   let jobChange = commit?.number;
+   if (jobChange === undefined) {
+      jobChange = stepRef?.change ?? job?.change;
+   }
 
+   // Need at least a change number to show the menu
    if (jobChange === undefined) {
-      jobChange = stepRef?.change ?? job.change!;
+      return null;
    }
 
    let change = jobChange.toString();
@@ -98,21 +99,21 @@ export const ChangeContextMenu: React.FC<{ target: ChangeContextMenuTarget, onDi
    }
    
    // Should only show cancel text if not in an incremental view where "commit" exists
-   if(!commit && job.abortedByUserInfo && job.cancellationReason) { 
+   if(!commit && job?.abortedByUserInfo && job?.cancellationReason) {
       menuItems.push({
          key: "cancelled",
          onRender: () => {
             return (
-               <Stack style={{ padding: 18, paddingTop: job.preflightDescription ? 0 : 18, paddingBottom: 18, maxWidth: 800 }} horizontal tokens={{ childrenGap: 12 }}>
+               <Stack style={{ padding: 18, paddingTop: job?.preflightDescription ? 0 : 18, paddingBottom: 18, maxWidth: 800 }} horizontal tokens={{ childrenGap: 12 }}>
                   <Text variant="small" style={{ fontFamily: "Horde Open Sans Bold" }}>Cancellation Reason:</Text>
-                  <Text variant="small" >{job.cancellationReason}</Text>
+                  <Text variant="small" >{job?.cancellationReason}</Text>
                </Stack>
             )
          }
       });
    }
 
-   if(job.cancellationReason || commit || job.preflightDescription){
+   if(job?.cancellationReason || commit || job?.preflightDescription){
       menuItems.push({
          key: "commit_divider",
          itemType: ContextualMenuItemType.Divider
@@ -129,24 +130,22 @@ export const ChangeContextMenu: React.FC<{ target: ChangeContextMenuTarget, onDi
 
       menuItems.push({ key: 'open_in_swarm', text: 'Open CL in Swarm', onClick: (ev) => { window.open(`${dashboard.swarmUrl}/changes/${change.replace("PF ", "")}`) } })
 
-      // range
-      const stream = projectStore.streamById(job.streamId)!;
-      const project = projectStore.byId(stream!.projectId)!;
-      const name = project.name === "Engine" ? "UE4" : project.name;
-
-
-      const historyUrl = `${dashboard.swarmUrl}/files/${name}/${stream.name}?range=@${jobChange}#commits`;
-      menuItems.push({ key: 'open_in_swarm_history', text: "Open CL History in Swarm", onClick: (ev) => { window.open(historyUrl) } })
-
-      let highCL = jobChange;
-      let lowCL = jobChange;
-
-      if (rangeCL !== undefined) {
-         lowCL = Math.min(jobChange, rangeCL);
-         highCL = Math.max(jobChange, rangeCL);
-         const url = `${dashboard.swarmUrl}/files/${name}/${stream.name}?range=@${lowCL},@${highCL}#commits`;
-         const rangeText = `Open CL Range ${lowCL} - ${highCL}`;
-         menuItems.push({ key: 'open_in_swarm_range', text: rangeText, onClick: (ev) => { window.open(url) } })
+      // Stream-dependent options (history, range) only available when we have stream info
+      const stream = job?.streamId ? projectStore.streamById(job.streamId) : undefined;
+      if (stream) {
+         const project = projectStore.byId(stream.projectId);
+         const name = project?.name === "Engine" ? "UE4" : project?.name;
+
+         const historyUrl = `${dashboard.swarmUrl}/files/${name}/${stream.name}?range=@${jobChange}#commits`;
+         menuItems.push({ key: 'open_in_swarm_history', text: "Open CL History in Swarm", onClick: (ev) => { window.open(historyUrl) } })
+
+         if (rangeCL !== undefined) {
+            const lowCL = Math.min(jobChange, rangeCL);
+            const highCL = Math.max(jobChange, rangeCL);
+            const url = `${dashboard.swarmUrl}/files/${name}/${stream.name}?range=@${lowCL},@${highCL}#commits`;
+            const rangeText = `Open CL Range ${lowCL} - ${highCL}`;
+            menuItems.push({ key: 'open_in_swarm_range', text: rangeText, onClick: (ev) => { window.open(url) } })
+         }
       }
    }
 
