Horde: Add NoReasonReported lease outcome reason

To handle cases where a lease is cancelled without an explicit reason from the agent.

diff --git a/Engine/Source/Programs/Horde/Plugins/Compute/HordeServer.Compute/Agents/AgentCollection.cs b/Engine/Source/Programs/Horde/Plugins/Compute/HordeServer.Compute/Agents/AgentCollection.cs
--- a/Engine/Source/Programs/Horde/Plugins/Compute/HordeServer.Compute/Agents/AgentCollection.cs
+++ b/Engine/Source/Programs/Horde/Plugins/Compute/HordeServer.Compute/Agents/AgentCollection.cs
@@ -985,7 +985,7 @@ namespace HordeServer.Agents
 					}
 					else
 					{
-						LeaseOutcomeReason reason = rpcLease?.ToNativeOutcomeReason() ?? LeaseOutcomeReason.None;
+						LeaseOutcomeReason reason = rpcLease?.ToNativeOutcomeReason() ?? LeaseOutcomeReason.NoReasonReported;
 						_logger.LogInformation("Setting outcome {Outcome} with reason {Reason} for lease {LeaseId} ({State})", rpcLease?.Outcome, reason, oldLease.Id, rpcLease?.State);
 						await _leaseCollection.TrySetOutcomeAsync(oldLease.Id, _clock.UtcNow, LeaseOutcome.Cancelled, reason, null, ToNativeLeaseMetrics(rpcLease?.Metrics), CancellationToken.None);
 					}
diff --git a/Engine/Source/Programs/Shared/EpicGames.Horde/Agents/Leases/LeaseMessages.cs b/Engine/Source/Programs/Shared/EpicGames.Horde/Agents/Leases/LeaseMessages.cs
--- a/Engine/Source/Programs/Shared/EpicGames.Horde/Agents/Leases/LeaseMessages.cs
+++ b/Engine/Source/Programs/Shared/EpicGames.Horde/Agents/Leases/LeaseMessages.cs
@@ -73,6 +73,9 @@ namespace EpicGames.Horde.Agents.Leases
 
 		/// <summary>Lease no longer needed despite scheduled/started</summary>
 		NoLongerNeeded = 512,
+
+		/// <summary>Lease was cancelled but no reason was reported by the agent</summary>
+		NoReasonReported = 1024,
 	}
 
 	/// <summary>
