Linking test result with test Health details view

#rnx

#Synced-CL 49992949

diff --git a/Engine/Source/Programs/Horde/HordeDashboard/plugins/test-automation/testAutomationCommon.tsx b/Engine/Source/Programs/Horde/HordeDashboard/plugins/test-automation/testAutomationCommon.tsx
--- a/Engine/Source/Programs/Horde/HordeDashboard/plugins/test-automation/testAutomationCommon.tsx
+++ b/Engine/Source/Programs/Horde/HordeDashboard/plugins/test-automation/testAutomationCommon.tsx
@@ -146,6 +146,18 @@ export const statusTexts = new Map<string, string>([
     [TestOutcome.Unspecified, "Catastrophic"]
 ]);
 
+export enum TestViewType {
+    /// The summary view
+    Summary = "0",
+    /// The health view
+    Health = "1"
+}
+
+export const testViewIcons = new Map<TestViewType, string>([
+    [TestViewType.Summary, "ShowResults"],
+    [TestViewType.Health, "Health"]
+]);
+
 export const MultiOptionChooser: React.FC<{options: IComboBoxOption[], initialSelection: string[], updateKeys: (selectedKeys: string[]) => void, placeholder?: string, style?: any, disabled?: boolean }> = ({ options, initialSelection, updateKeys, placeholder, style, disabled }) => {
     const [selection, setSelection] = useState<string[]>(initialSelection);
     const optionsWithSelectAll = [...options];
diff --git a/Engine/Source/Programs/Horde/HordeDashboard/plugins/test-automation/testAutomationPhaseView.tsx b/Engine/Source/Programs/Horde/HordeDashboard/plugins/test-automation/testAutomationPhaseView.tsx
--- a/Engine/Source/Programs/Horde/HordeDashboard/plugins/test-automation/testAutomationPhaseView.tsx
+++ b/Engine/Source/Programs/Horde/HordeDashboard/plugins/test-automation/testAutomationPhaseView.tsx
@@ -10,7 +10,7 @@ import { ChangeButton } from "horde/components/ChangeButton";
 import { getShortNiceTime, msecToElapsed } from "horde/base/utilities/timeUtils";
 import { getHordeTheme } from "horde/styles/theme";
 import { GetArtifactResponse } from "horde/backend/Api";
-import { CheckListOption, getPhaseSessionStatusColor, phaseFiltersContext, PhaseHistoryGraphWidget, PhasesStatusPie, styles, tagColors } from "./testAutomationCommon";
+import { CheckListOption, getPhaseSessionStatusColor, phaseFiltersContext, PhaseHistoryGraphWidget, PhasesStatusPie, styles, tagColors, testViewIcons, TestViewType } from "./testAutomationCommon";
 import { projectStore } from "horde/backend/ProjectStore";
 import { JobArtifactsModal } from "horde/components/artifacts/ArtifactsModal";
 import { getJobParams, JobContext } from "./phaseCompare";
@@ -491,6 +491,15 @@ export const TestPhasesView: React.FC<{ testDataId: string, handler: TestDataHan
                                                 sessionDetails.artifacts?.getJobStepSavedArtifacts().then(value => setJobArtifacts(value));
                                             }}>Artifacts</DefaultButton>
                                     </Stack>
+                                    {!!sessionDetails.test && !!handler.stream &&
+                                        <Link style={{padding: '0px 5px'}} title="open test health view" onClick={() => {
+                                                handler.setSearchParam('view', TestViewType.Health);
+                                                handler.setSearchParam('health', sessionDetails.test!.id);
+                                                handler.removeSearchParam('session');
+                                            }}>
+                                            <FontIcon iconName={testViewIcons.get(TestViewType.Health)} style={{fontSize: 15, paddingRight: 4}} />
+                                        </Link>
+                                    }
                                 </Stack>
                             </Stack>
 
diff --git a/Engine/Source/Programs/Horde/HordeDashboard/plugins/test-automation/testAutomationSummary.tsx b/Engine/Source/Programs/Horde/HordeDashboard/plugins/test-automation/testAutomationSummary.tsx
--- a/Engine/Source/Programs/Horde/HordeDashboard/plugins/test-automation/testAutomationSummary.tsx
+++ b/Engine/Source/Programs/Horde/HordeDashboard/plugins/test-automation/testAutomationSummary.tsx
@@ -1,6 +1,6 @@
 // Copyright Epic Games, Inc. All Rights Reserved.
 
-import { DefaultButton, FontIcon, IconButton, mergeStyleSets, Modal, Stack, Text, HoverCard, HoverCardType, IPlainCardProps, DirectionalHint, Spinner, SpinnerSize, IComboBoxOption } from "@fluentui/react";
+import { DefaultButton, FontIcon, IconButton, mergeStyleSets, Modal, Stack, Text, HoverCard, HoverCardType, IPlainCardProps, DirectionalHint, Spinner, SpinnerSize, IComboBoxOption, Link } from "@fluentui/react";
 import { useState, useEffect, useCallback, memo } from "react";
 import { observer } from "mobx-react-lite";
 import { TestDataHandler, TestNameRef, MetadataRef, TestSessionStatus, TestSessionResult, TestStatus, TestTagRef } from "./testData";
@@ -9,7 +9,7 @@ import dashboard, { StatusColor } from "horde/backend/Dashboard";
 import { getShortNiceTime } from "horde/base/utilities/timeUtils";
 import { getHordeStyling } from "horde/styles/Styles";
 import { projectStore } from 'horde/backend/ProjectStore';
-import { SessionStatusBar, SessionValues, StatusBar, StatusBarStack, MultiOptionChooser, getTestSessionStatusColor, getStatusColors, statusTexts, getTestUniqueMetaIdentifiers, StreamSelector, styles } from "./testAutomationCommon";
+import { SessionStatusBar, SessionValues, StatusBar, StatusBarStack, MultiOptionChooser, getTestSessionStatusColor, getStatusColors, statusTexts, getTestUniqueMetaIdentifiers, StreamSelector, styles, TestViewType, testViewIcons } from "./testAutomationCommon";
 import { TestHistoryGraph } from "./testHistoryGraph";
 import { StoreKey, userStorage } from "./testAutomationUserStorage";
 
@@ -61,14 +61,14 @@ export const TestHistoryView: React.FC<{ test: TestNameRef, handler: TestDataHan
     handler.subscribeToSubQueryLoading();
 
     const onClickStream = useCallback((targetStream: string) => {
-        setStreams({...streams, toCompare: undefined});
+        setStreams(prev => ({...prev, toCompare: undefined}));
         handler.query(targetStream, true)
             .catch((reason) => console.error(reason))
             .finally(() => {
-                setStreams({...streams, toCompare: targetStream});
+                setStreams(prev => ({...prev, toCompare: targetStream}));
                 userStorage.setItem(StoreKey.HistoryToCompare, targetStream);
             });
-    }, [streams]);
+    }, []);
 
     const availableStreams = handler.availableStreams.filter((s) => s !== streams.current);
 
@@ -130,10 +130,16 @@ export const TestHistoryView: React.FC<{ test: TestNameRef, handler: TestDataHan
     return <Stack>
             <Stack className={hordeClasses.raised}>
                 <Stack>
-                    <Stack horizontal>
+                    <Stack horizontal tokens={{childrenGap: 10}}>
                         <Stack>
                             <Text style={{ fontFamily: "Horde Open Sans Semibold" }} variant='medium'>{name}</Text>
                         </Stack>
+                        <Link style={{padding: '0px 5px'}} title="open test health view" onClick={() => {
+                                handler.setSearchParam('view', TestViewType.Health);
+                                handler.setSearchParam('health', test.id);
+                            }}>
+                            <FontIcon iconName={testViewIcons.get(TestViewType.Health)} style={{fontSize: 15, paddingRight: 4}} />
+                        </Link>
                     </Stack>
                 </Stack>
                 <Stack horizontal tokens={{childrenGap: 15}} style={{ paddingBottom: 12, paddingTop: 5 }}>
diff --git a/Engine/Source/Programs/Horde/HordeDashboard/plugins/test-automation/testAutomationView.tsx b/Engine/Source/Programs/Horde/HordeDashboard/plugins/test-automation/testAutomationView.tsx
--- a/Engine/Source/Programs/Horde/HordeDashboard/plugins/test-automation/testAutomationView.tsx
+++ b/Engine/Source/Programs/Horde/HordeDashboard/plugins/test-automation/testAutomationView.tsx
@@ -10,7 +10,7 @@ import React, { useEffect, useState, useCallback, useMemo } from "react";
 import { projectStore } from 'horde/backend/ProjectStore';
 import { TestDataHandler, TestDataVersionRegistrar, TestSessionDetails } from "./testData";
 import { TestSummary } from "./testAutomationSummary";
-import { getStatusColors, getStreamOptions, getTestSessionStatusColor, MultiOptionChooser, SessionStatusBar, tagColors } from "./testAutomationCommon";
+import { getStatusColors, getStreamOptions, getTestSessionStatusColor, MultiOptionChooser, SessionStatusBar, tagColors, testViewIcons, TestViewType } from "./testAutomationCommon";
 import { TestPhasesView } from "./testAutomationPhaseView";
 import { JobStepTestDataItem } from "./api";
 
@@ -282,7 +282,7 @@ const TestAutomationSidebarLeft: React.FC = observer(() => {
     filterMetaKeys = filterMetaKeys.concat(handler.filterState.testmeta?.keys().toArray() ?? []);
     filterMetaKeys.sort((a, b) => a.localeCompare(b));
 
-    const view = (handler.getSearchParam("view") as string | undefined) ?? "0";
+    const view = (handler.getSearchParam("view") as string | undefined) ?? TestViewType.Summary;
 
     const [metaKeys, setMetaKeys] = useState<string[]>(filterMetaKeys);
     const { hordeClasses } = getHordeStyling();
@@ -305,8 +305,8 @@ const TestAutomationSidebarLeft: React.FC = observer(() => {
     }, [handler.searchUpdated]);
 
    return <Stack style={{ width: 300, paddingRight: 18 }} className={hordeClasses.modal}>
-            <Pivot selectedKey={view} onLinkClick={(item) => {handler.setSearchParam("view", item?.props?.itemKey ?? "0")}}>
-                <PivotItem headerText="Summary" itemKey="0">
+            <Pivot selectedKey={view} onLinkClick={(item) => {handler.setSearchParam("view", item?.props?.itemKey ?? TestViewType.Summary)}}>
+                <PivotItem headerText="Latest Status" itemKey={TestViewType.Summary} itemIcon={testViewIcons.get(TestViewType.Summary)}>
                     <Stack key="chooser_stream">
                         <StreamChooser />
                     </Stack>
@@ -321,7 +321,7 @@ const TestAutomationSidebarLeft: React.FC = observer(() => {
                         <MetaKeyChooser key="chooser_meta" onUpdateKey={updateMetaKey} />
                     </Stack>
                 </PivotItem>
-                <PivotItem headerText="Health" itemKey="1">
+                <PivotItem headerText="Health" itemKey={TestViewType.Health} itemIcon={testViewIcons.get(TestViewType.Health)}>
                     <Stack key="chooser_stream">
                         <StreamChooser />
                     </Stack>
