[Horde][Config] Updating default view to align with telemetry

#Synced-CL 49930492

diff --git a/Engine/Source/Programs/Horde/HordeServer/Defaults/default-views.telemetry.json b/Engine/Source/Programs/Horde/HordeServer/Defaults/default-views.telemetry.json
--- a/Engine/Source/Programs/Horde/HordeServer/Defaults/default-views.telemetry.json
+++ b/Engine/Source/Programs/Horde/HordeServer/Defaults/default-views.telemetry.json
@@ -600,6 +600,16 @@
 								}
 							]
 						},
+						{
+							"name": "UB Duration",
+							"display": "Value",
+							"metrics": [
+								{
+									"id": "ubt-completed-duration",
+									"alias": "FullDuration"
+								}
+							]
+						},
 						{
 							"name": "Actions All (Sum)",
 							"display": "Value",
