[Horde][Config] removing debug probe

#Synced-CL 49930492

diff --git a/Engine/Source/Programs/Horde/HordeServer/Defaults/default-metrics.telemetry.json b/Engine/Source/Programs/Horde/HordeServer/Defaults/default-metrics.telemetry.json
--- a/Engine/Source/Programs/Horde/HordeServer/Defaults/default-metrics.telemetry.json
+++ b/Engine/Source/Programs/Horde/HordeServer/Defaults/default-metrics.telemetry.json
@@ -850,14 +850,6 @@
 			"groupBy": "$.Payload.Stream,$.Payload.Version",
 			"function": "Count",
 			"interval": "00:30:00"
-		},
-		{
-			"id": "debug-probe",
-			"filter": "$[?(@.Payload.SchemaVersion >= 0 )]",
-			"groupBy": "$.Payload.Build_BranchName,$.Payload.Session_Label,$.Payload.ProjectName,$.Payload.EventName",
-			"property": "$.Payload.SchemaVersion",
-			"function": "Count",
-			"interval": "00:05:00"
 		}
 	]
 }
\ No newline at end of file
diff --git a/Engine/Source/Programs/Horde/HordeServer/Defaults/default-views.telemetry.json b/Engine/Source/Programs/Horde/HordeServer/Defaults/default-views.telemetry.json
--- a/Engine/Source/Programs/Horde/HordeServer/Defaults/default-views.telemetry.json
+++ b/Engine/Source/Programs/Horde/HordeServer/Defaults/default-views.telemetry.json
@@ -797,15 +797,6 @@
 				{
 					"name": "Debug",
 					"charts": [
-						{
-							"name": "Debug",
-							"display": "Value",
-							"metrics": [
-								{
-									"id": "debug-probe"
-								}
-							]
-						},
 						{
 							"name": "MountPoint Test",
 							"display": "Value",
