[Horde][Server] Exception logging for background process in metrics collection.

#Synced-CL 49930492

diff --git a/Engine/Source/Programs/Horde/Plugins/Analytics/HordeServer.Analytics/Telemetry/Metrics/MetricCollection.cs b/Engine/Source/Programs/Horde/Plugins/Analytics/HordeServer.Analytics/Telemetry/Metrics/MetricCollection.cs
--- a/Engine/Source/Programs/Horde/Plugins/Analytics/HordeServer.Analytics/Telemetry/Metrics/MetricCollection.cs
+++ b/Engine/Source/Programs/Horde/Plugins/Analytics/HordeServer.Analytics/Telemetry/Metrics/MetricCollection.cs
@@ -148,7 +148,14 @@ namespace HordeServer.Telemetry.Metrics
 		{
 			await foreach ((TelemetryStoreId storeId, JsonNode node) in _events.Reader.ReadAllAsync(cancellationToken))
 			{
-				ProcessEvent(storeId, node);
+				try
+				{
+					ProcessEvent(storeId, node);
+				}
+				catch (Exception ex)
+				{
+					_logger.LogError(ex, "Failed to process telemetry event for store {StoreId}", storeId);
+				}
 			}
 		}
 
@@ -337,38 +344,49 @@ namespace HordeServer.Telemetry.Metrics
 
 			while (!cancellationToken.IsCancellationRequested)
 			{
-				await newDataTask.WaitAsync(cancellationToken);
-				await Task.WhenAny(flushTask, Task.Delay(TimeSpan.FromSeconds(5.0), cancellationToken));
-
-				newDataTask = _newDataEvent.Task;
-				flushTask = _flushEvent.Task;
-
 				try
 				{
-					await WriteQueuedSamplesAsync(cancellationToken);
-				}
-				catch (Exception ex)
-				{
-					_logger.LogError(ex, "Error flushing telemetry data: {Message}", ex.Message);
-				}
+					await newDataTask.WaitAsync(cancellationToken);
+					await Task.WhenAny(flushTask, Task.Delay(TimeSpan.FromSeconds(5.0), cancellationToken));
 
-				TimeSpan slowMetricReportTime = TimeSpan.FromMinutes(1.0);
-				if (timer.Elapsed > slowMetricReportTime)
-				{
-					int numDroppedItems = Interlocked.Exchange(ref _numDroppedItems, 0);
-					if (numDroppedItems > 0)
+					newDataTask = _newDataEvent.Task;
+					flushTask = _flushEvent.Task;
+
+					try
 					{
-						lock (_lockObject)
+						await WriteQueuedSamplesAsync(cancellationToken);
+					}
+					catch (Exception ex)
+					{
+						_logger.LogError(ex, "Error flushing telemetry data: {Message}", ex.Message);
+					}
+
+					TimeSpan slowMetricReportTime = TimeSpan.FromMinutes(1.0);
+					if (timer.Elapsed > slowMetricReportTime)
+					{
+						int numDroppedItems = Interlocked.Exchange(ref _numDroppedItems, 0);
+						if (numDroppedItems > 0)
 						{
-							_logger.LogWarning("{NumEvents} telemetry events were dropped due to queue being full.", numDroppedItems);
-							foreach ((QualifiedMetricId qualifiedMetricId, TimeSpan time) in _slowMetrics)
+							lock (_lockObject)
 							{
-								_logger.LogWarning("Metric {StoreId}:{MetricId} took {TimeSecs:n1}s ({Pct:n1}%) of analysis time in the last minute", qualifiedMetricId.StoreId, qualifiedMetricId.MetricId, time.TotalSeconds, time.TotalSeconds / slowMetricReportTime.TotalSeconds); 
+								_logger.LogWarning("{NumEvents} telemetry events were dropped due to queue being full.", numDroppedItems);
+								foreach ((QualifiedMetricId qualifiedMetricId, TimeSpan time) in _slowMetrics)
+								{
+									_logger.LogWarning("Metric {StoreId}:{MetricId} took {TimeSecs:n1}s ({Pct:n1}%) of analysis time in the last minute", qualifiedMetricId.StoreId, qualifiedMetricId.MetricId, time.TotalSeconds, time.TotalSeconds / slowMetricReportTime.TotalSeconds);
+								}
+								_slowMetrics.Clear();
 							}
-							_slowMetrics.Clear();
 						}
+						timer.Restart();
 					}
-					timer.Restart();
+				}
+				catch (OperationCanceledException)
+				{
+					break;
+				}
+				catch (Exception ex)
+				{
+					_logger.LogError(ex, "Error in telemetry flush loop.");
 				}
 			}
 		}
