Horde: Refactor lease outcome reason merging

Also improve server-initiated termination handling by passing better defaults.

diff --git a/Engine/Source/Programs/Horde/Plugins/Compute/HordeServer.Compute.Tests/Agents/AgentCollectionTests.cs b/Engine/Source/Programs/Horde/Plugins/Compute/HordeServer.Compute.Tests/Agents/AgentCollectionTests.cs
--- a/Engine/Source/Programs/Horde/Plugins/Compute/HordeServer.Compute.Tests/Agents/AgentCollectionTests.cs
+++ b/Engine/Source/Programs/Horde/Plugins/Compute/HordeServer.Compute.Tests/Agents/AgentCollectionTests.cs
@@ -10,6 +10,7 @@ using HordeAgent::HordeCommon.Rpc.Tasks;
 using HordeCommon.Rpc;
 using HordeCommon.Rpc.Messages;
 using HordeServer.Agents;
+using HordeServer.Agents.Leases;
 using System.Security.Cryptography;
 using EpicGames.Horde.Tools;
 using Google.Protobuf;
@@ -319,6 +320,65 @@ public class AgentCollectionTests : ComputeTestSetup
 		Assert.AreEqual(0, leases.Length);
 	}
 
+	[TestMethod]
+	public async Task TerminateSession_SetsLeaseOutcomeReasonAsync()
+	{
+		// Create session first
+		IAgent? agent = await _agent.TryCreateSessionAsync(new CreateSessionOptions(new RpcAgentCapabilities(), Array.Empty<PoolId>(), "foo"));
+		Assert.IsNotNull(agent);
+		Assert.IsNotNull(agent.SessionId);
+
+		// Create leases
+		agent = await agent.TryCreateLeaseAsync(_lease1);
+		Assert.IsNotNull(agent);
+
+		agent = await agent.TryCreateLeaseAsync(_lease2);
+		Assert.IsNotNull(agent);
+
+		// Terminate session without agent-provided lease data (server-initiated termination)
+		await agent.TryTerminateSessionAsync();
+
+		// Verify leases have correct outcome reason
+		ILease? lease1 = await LeaseCollection.GetAsync(_lease1.Id);
+		ILease? lease2 = await LeaseCollection.GetAsync(_lease2.Id);
+
+		Assert.IsNotNull(lease1);
+		Assert.IsNotNull(lease2);
+
+		// Both leases should be cancelled with ServerInitiated | AgentSessionTerminated reason
+		Assert.AreEqual(LeaseOutcome.Cancelled, lease1.Outcome);
+		Assert.AreEqual(LeaseOutcome.Cancelled, lease2.Outcome);
+		Assert.IsTrue(lease1.OutcomeReason.HasFlag(LeaseOutcomeReason.ServerInitiated));
+		Assert.IsTrue(lease1.OutcomeReason.HasFlag(LeaseOutcomeReason.AgentSessionTerminated));
+		Assert.IsTrue(lease2.OutcomeReason.HasFlag(LeaseOutcomeReason.ServerInitiated));
+		Assert.IsTrue(lease2.OutcomeReason.HasFlag(LeaseOutcomeReason.AgentSessionTerminated));
+	}
+	
+	[TestMethod]
+	public async Task UpdateSession_MergesOutcomeReasons_WhenLeaseAlreadyFinishedAsync()
+	{
+		IAgent? agent = await _agent.TryCreateSessionAsync(new CreateSessionOptions(new RpcAgentCapabilities(), [], "foo"));
+		agent = await agent!.TryCreateLeaseAsync(_lease1);
+		agent = await agent!.TryUpdateSessionAsync(new UpdateSessionOptions { Leases = [new() { Id = _lease1.Id, State = RpcLeaseState.Active }] });
+
+		// Simulate race: Another process already finished the lease
+		bool result = await LeaseCollection.TrySetOutcomeAsync(_lease1.Id, DateTime.UtcNow, LeaseOutcome.Cancelled,
+			LeaseOutcomeReason.ServerInitiated | LeaseOutcomeReason.AgentSessionTerminated, null, null);
+		Assert.IsTrue(result);
+
+		// Agent reports lease with additional reason - should merge via TryAddOutcomeReasonsAsync
+		RpcLease cancelledLease = new() { Id = _lease1.Id, State = RpcLeaseState.Cancelled, Outcome = RpcLeaseOutcome.Cancelled };
+		cancelledLease.OutcomeReasons.Add(RpcLeaseOutcomeReason.NonZeroExitCode);
+		agent = await agent!.TryUpdateSessionAsync(new UpdateSessionOptions { Leases = [cancelledLease] });
+		Assert.IsNotNull(agent);
+
+		// Verify all reasons merged
+		ILease? lease = await LeaseCollection.GetAsync(_lease1.Id);
+		Assert.IsTrue(lease!.OutcomeReason.HasFlag(LeaseOutcomeReason.ServerInitiated));
+		Assert.IsTrue(lease.OutcomeReason.HasFlag(LeaseOutcomeReason.AgentSessionTerminated));
+		Assert.IsTrue(lease.OutcomeReason.HasFlag(LeaseOutcomeReason.NonZeroExitCode));
+	}
+
 	[TestMethod]
 	public async Task GetChildLeaseIdsAsync()
 	{
diff --git a/Engine/Source/Programs/Horde/Plugins/Compute/HordeServer.Compute/Agents/AgentCollection.cs b/Engine/Source/Programs/Horde/Plugins/Compute/HordeServer.Compute/Agents/AgentCollection.cs
--- a/Engine/Source/Programs/Horde/Plugins/Compute/HordeServer.Compute/Agents/AgentCollection.cs
+++ b/Engine/Source/Programs/Horde/Plugins/Compute/HordeServer.Compute/Agents/AgentCollection.cs
@@ -344,7 +344,7 @@ namespace HordeServer.Agents
 					if (agent != null && agent.Session != null && agent.Session.UpdateTicks == session.UpdateTicks)
 					{
 						_logger.LogDebug("Terminating session {SessionId} for agent {Agent} (expiry time: {Time})", session.SessionId, session.AgentId, agent.Session?.ExpiryTime);
-						await TryUpdateSessionAsync(agent, new UpdateSessionOptions { Status = AgentStatus.Stopped }, cancellationToken);
+						await TryUpdateSessionAsync(agent, new UpdateSessionOptions { Status = AgentStatus.Stopped, DefaultLeaseOutcomeReason = LeaseOutcomeReason.ServerInitiated | LeaseOutcomeReason.AgentSessionTerminated }, cancellationToken);
 						c++;
 					}
 				}
@@ -974,20 +974,33 @@ namespace HordeServer.Agents
 			// Remove any complete leases
 			if (newSession.Leases != null)
 			{
+				// Merges outcome reasons if lease was already finished by another process (race condition)
+				async Task TryMergeOutcomeReasonsAsync(bool updated, LeaseId leaseId, LeaseOutcomeReason reason)
+				{
+					if (!updated && reason != LeaseOutcomeReason.None)
+					{
+						_logger.LogInformation("Merging outcome reasons {Reason} for already-finished lease {LeaseId}", reason, leaseId);
+						await _leaseCollection.TryAddOutcomeReasonsAsync(leaseId, reason, CancellationToken.None);
+					}
+				}
+
 				HashSet<LeaseId> newLeaseIds = new HashSet<LeaseId>(newSession.Leases?.Select(x => x.Id) ?? Enumerable.Empty<LeaseId>());
 				foreach (RpcSessionLease oldLease in agent.Session.Leases.Where(x => !newLeaseIds.Contains(x.Id)))
 				{
 					RpcLease? rpcLease = options.Leases?.FirstOrDefault(x => x.Id == oldLease.Id);
 					if (rpcLease is { State: RpcLeaseState.Completed })
 					{
+						LeaseOutcomeReason reason = rpcLease.ToNativeOutcomeReason();
 						_logger.LogInformation("Setting outcome {Outcome} for completed lease {LeaseId}", oldLease.Id, rpcLease.Outcome);
-						await _leaseCollection.TrySetOutcomeAsync(oldLease.Id, _clock.UtcNow, (LeaseOutcome)rpcLease.Outcome, LeaseOutcomeReason.None, rpcLease.Output.ToArray(), ToNativeLeaseMetrics(rpcLease.Metrics), CancellationToken.None);
+						bool updated = await _leaseCollection.TrySetOutcomeAsync(oldLease.Id, _clock.UtcNow, (LeaseOutcome)rpcLease.Outcome, reason, rpcLease.Output.ToArray(), ToNativeLeaseMetrics(rpcLease.Metrics), CancellationToken.None);
+						await TryMergeOutcomeReasonsAsync(updated, oldLease.Id, reason);
 					}
 					else
 					{
-						LeaseOutcomeReason reason = rpcLease?.ToNativeOutcomeReason() ?? LeaseOutcomeReason.NoReasonReported;
+						LeaseOutcomeReason reason = rpcLease?.ToNativeOutcomeReason() ?? options.DefaultLeaseOutcomeReason ?? LeaseOutcomeReason.NoReasonReported;
 						_logger.LogInformation("Setting outcome {Outcome} with reason {Reason} for lease {LeaseId} ({State})", rpcLease?.Outcome, reason, oldLease.Id, rpcLease?.State);
-						await _leaseCollection.TrySetOutcomeAsync(oldLease.Id, _clock.UtcNow, LeaseOutcome.Cancelled, reason, null, ToNativeLeaseMetrics(rpcLease?.Metrics), CancellationToken.None);
+						bool updated = await _leaseCollection.TrySetOutcomeAsync(oldLease.Id, _clock.UtcNow, LeaseOutcome.Cancelled, reason, null, ToNativeLeaseMetrics(rpcLease?.Metrics), CancellationToken.None);
+						await TryMergeOutcomeReasonsAsync(updated, oldLease.Id, reason);
 					}
 				}
 			}
@@ -1170,7 +1183,7 @@ namespace HordeServer.Agents
 		async Task<Agent?> TryTerminateSessionAsync(Agent agent, CancellationToken cancellationToken)
 		{
 			_logger.LogInformation("Terminating session {SessionId} for {AgentId}", agent.Session?.SessionId ?? SessionId.Empty, agent.Id);
-			return await TryUpdateSessionAsync(agent, new UpdateSessionOptions(AgentStatus.Stopped), cancellationToken);
+			return await TryUpdateSessionAsync(agent, new UpdateSessionOptions(AgentStatus.Stopped, DefaultLeaseOutcomeReason: LeaseOutcomeReason.ServerInitiated | LeaseOutcomeReason.AgentSessionTerminated), cancellationToken);
 		}
 
 		/// <inheritdoc/>
diff --git a/Engine/Source/Programs/Horde/Plugins/Compute/HordeServer.Compute/Agents/IAgent.cs b/Engine/Source/Programs/Horde/Plugins/Compute/HordeServer.Compute/Agents/IAgent.cs
--- a/Engine/Source/Programs/Horde/Plugins/Compute/HordeServer.Compute/Agents/IAgent.cs
+++ b/Engine/Source/Programs/Horde/Plugins/Compute/HordeServer.Compute/Agents/IAgent.cs
@@ -369,7 +369,8 @@ namespace HordeServer.Agents
 	/// <param name="Capabilities">Capbilities for the session</param>
 	/// <param name="DynamicPools">New list of dynamic pools for the agent</param>
 	/// <param name="Leases">New set of leases</param>
-	public record UpdateSessionOptions(AgentStatus? Status = null, RpcAgentCapabilities? Capabilities = null, IReadOnlyList<PoolId>? DynamicPools = null, IEnumerable<RpcLease>? Leases = null);
+	/// <param name="DefaultLeaseOutcomeReason">Default reason used when terminating leases without agent-provided reasons</param>
+	public record UpdateSessionOptions(AgentStatus? Status = null, RpcAgentCapabilities? Capabilities = null, IReadOnlyList<PoolId>? DynamicPools = null, IEnumerable<RpcLease>? Leases = null, LeaseOutcomeReason? DefaultLeaseOutcomeReason = null);
 
 	/// <summary>
 	/// Extension methods for IAgent
diff --git a/Engine/Source/Programs/Horde/Plugins/Compute/HordeServer.Compute/Agents/Leases/ILeaseCollection.cs b/Engine/Source/Programs/Horde/Plugins/Compute/HordeServer.Compute/Agents/Leases/ILeaseCollection.cs
--- a/Engine/Source/Programs/Horde/Plugins/Compute/HordeServer.Compute/Agents/Leases/ILeaseCollection.cs
+++ b/Engine/Source/Programs/Horde/Plugins/Compute/HordeServer.Compute/Agents/Leases/ILeaseCollection.cs
@@ -107,5 +107,14 @@ namespace HordeServer.Agents.Leases
 		/// <param name="cancellationToken">Cancellation token for the operation</param>
 		/// <returns>True if the lease was updated, false otherwise</returns>
 		Task<bool> TrySetOutcomeAsync(LeaseId leaseId, DateTime finishTime, LeaseOutcome outcome, LeaseOutcomeReason outcomeReason, byte[]? output, LeaseMetrics? metrics, CancellationToken cancellationToken = default);
+
+		/// <summary>
+		/// Adds outcome reasons to an existing lease. This can be used to merge reasons even after the lease has finished.
+		/// </summary>
+		/// <param name="leaseId">The lease to update</param>
+		/// <param name="outcomeReason">Additional reasons to add (will be ORed with existing reasons)</param>
+		/// <param name="cancellationToken">Cancellation token for the operation</param>
+		/// <returns>True if the lease was updated, false if lease doesn't exist</returns>
+		Task<bool> TryAddOutcomeReasonsAsync(LeaseId leaseId, LeaseOutcomeReason outcomeReason, CancellationToken cancellationToken = default);
 	}
 }
diff --git a/Engine/Source/Programs/Horde/Plugins/Compute/HordeServer.Compute/Agents/Leases/LeaseCollection.cs b/Engine/Source/Programs/Horde/Plugins/Compute/HordeServer.Compute/Agents/Leases/LeaseCollection.cs
--- a/Engine/Source/Programs/Horde/Plugins/Compute/HordeServer.Compute/Agents/Leases/LeaseCollection.cs
+++ b/Engine/Source/Programs/Horde/Plugins/Compute/HordeServer.Compute/Agents/Leases/LeaseCollection.cs
@@ -264,5 +264,20 @@ namespace HordeServer.Agents.Leases
 			OnLeaseComplete?.Invoke(document);
 			return true;
 		}
+
+		/// <inheritdoc/>
+		public async Task<bool> TryAddOutcomeReasonsAsync(LeaseId leaseId, LeaseOutcomeReason outcomeReason, CancellationToken cancellationToken)
+		{
+			if (outcomeReason == LeaseOutcomeReason.None)
+			{
+				return true; // Nothing to add
+			}
+
+			FilterDefinition<LeaseDocument> filter = Builders<LeaseDocument>.Filter.Eq(x => x.Id, leaseId);
+			UpdateDefinition<LeaseDocument> update = Builders<LeaseDocument>.Update.BitwiseOr(x => x.OutcomeReason, outcomeReason);
+
+			UpdateResult result = await _leases.UpdateOneAsync(filter, update, cancellationToken: cancellationToken);
+			return result.MatchedCount > 0;
+		}
 	}
 }
