[Horde][Analytics][Server] Reducing the severity of noisy error on duplicate payload properties due to the benign nature.

#Synced-CL 49930492

diff --git a/Engine/Source/Programs/Horde/Plugins/Analytics/HordeServer.Analytics/Telemetry/Metrics/MetricCollection.cs b/Engine/Source/Programs/Horde/Plugins/Analytics/HordeServer.Analytics/Telemetry/Metrics/MetricCollection.cs
--- a/Engine/Source/Programs/Horde/Plugins/Analytics/HordeServer.Analytics/Telemetry/Metrics/MetricCollection.cs
+++ b/Engine/Source/Programs/Horde/Plugins/Analytics/HordeServer.Analytics/Telemetry/Metrics/MetricCollection.cs
@@ -152,8 +152,15 @@ namespace HordeServer.Telemetry.Metrics
 				{
 					ProcessEvent(storeId, node);
 				}
-				catch (Exception ex)
+				catch (ArgumentException ex)
+				{
+					// This is currently a narrow exception capture for a common event where malformed payloads are passed the background processor.
+					// We will emit this as a logDebug along with the exception while we fix call sites that are malformed.
+					_logger.LogDebug(ex, "Skipping telemetry emission for event store {StoreId}; processing failed due to malformed payload.", storeId);
+				}
+				catch(Exception ex)
 				{
+					// All other exceptions should be logged as errors.
 					_logger.LogError(ex, "Failed to process telemetry event for store {StoreId}", storeId);
 				}
 			}
