Horde: Update release notes

#rnx

diff --git a/Engine/Source/Programs/Horde/Docs/ReleaseNotes.md b/Engine/Source/Programs/Horde/Docs/ReleaseNotes.md
--- a/Engine/Source/Programs/Horde/Docs/ReleaseNotes.md
+++ b/Engine/Source/Programs/Horde/Docs/ReleaseNotes.md
@@ -1,5 +1,38 @@
 # Release Notes
 
+## 2026-01-22
+
+* Reducing the severity of noisy error on duplicate payload properties due to the benign nature. (50024919)
+* Refactor lease outcome reason merging
+  Also improve server-initiated termination handling by passing better defaults. (50023636)
+* Exception logging for background process in metrics collection. (50018209)
+* removing debug probe (50007684)
+* Updating default view to align with telemetry (50006350)
+* Add NoReasonReported lease outcome reason
+  To handle cases where a lease is cancelled without an explicit reason from the agent. (49992140)
+* Add support for change button to not require job object (49969938)
+* Enable cache when requesting access token via EpicGames.OIDC
+  Uses new caching functionality from 49949594 (49950769)
+* Log reason for job driver performing cleanup (49950560)
+* Expose LeaseOutcomeReason in GetJobBatchResponse
+  Only fetches lease data when batch has an error to avoid unnecessary DB queries for successful batches. (49949816)
+* EpicGames.OIDC: Add opt-in access token caching
+  Currently only refresh tokens are cached to disk. This change adds the ability to also cache access tokens, reducing unnecessary calls when tokens are still valid.
+  Still writes refresh tokens in the old format as "Providers" key, while also writing a new "Providers2" key which includes a cached access tokens. (49949594)
+* Temporarily removing constraint of unused key in order to update sparse. (49930492)
+* Add diagnostic logging to schedule gate job searches
+  Track why jobs are skipped (failed, incomplete, or missing target) and log summary when all candidates are exhausted. (49899154)
+* Introducing key stat performance records. (49898834)
+* Much better dashboard commit cache behavior (49897270)
+* Minor improvements to nullability check & safe fallback for ColumnAttribute with no name (49893700)
+* Expanding capabilities of Analytics.Telemetry to simplify telemetry emissions & reduce duplication keys across code. (49892781)
+* Removal of sparse constraint as it's less effective on queryPlanner version 1 (49890466)
+* Add retry logic to TestAgentMessageLoopTcpAesAsync to handle CI load variability (49884945)
+* Fix compiler warnings for JobDriver.Tests after upgrading MSTest (49884863)
+* Adding new Performance Trends c# project (49874115)
+* Exposing existing Session ID to the telemetry api (49862415)
+* Adds IPluginDashboardAcls interface for plugins to control dashboard feature visibility based on user permissions (49857922)
+
 ## 2026-01-15
 
 * Address issues with on prem Jira service build health (49818828)
