Horde: Propagate lease cancellation reason for cancelled job batches

diff --git a/Engine/Source/Programs/Horde/Plugins/Build/HordeServer.Build/Jobs/JobService.cs b/Engine/Source/Programs/Horde/Plugins/Build/HordeServer.Build/Jobs/JobService.cs
--- a/Engine/Source/Programs/Horde/Plugins/Build/HordeServer.Build/Jobs/JobService.cs
+++ b/Engine/Source/Programs/Horde/Plugins/Build/HordeServer.Build/Jobs/JobService.cs
@@ -383,7 +383,16 @@ namespace HordeServer.Jobs
 					{
 						if ((batch.Error == JobStepBatchError.Cancelled || batch.Error == JobStepBatchError.NoLongerNeeded) && (batch.State == JobStepBatchState.Starting || batch.State == JobStepBatchState.Running) && batch.AgentId != null && batch.LeaseId != null)
 						{
-							await CancelLeaseAsync(batch.AgentId.Value, batch.LeaseId.Value, cancellationToken);
+							LeaseOutcomeReason reason = LeaseOutcomeReason.ServerInitiated;
+							if (abortedByUserId != null)
+							{
+								reason |= LeaseOutcomeReason.UserInitiated;
+							}
+							if (batch.Error == JobStepBatchError.NoLongerNeeded)
+							{
+								reason |= LeaseOutcomeReason.NoLongerNeeded;
+							}
+							await CancelLeaseAsync(batch.AgentId.Value, batch.LeaseId.Value, reason, cancellationToken);
 						}
 					}
 
@@ -473,13 +482,15 @@ namespace HordeServer.Jobs
 		/// </summary>
 		/// <param name="agentId">The agent to retreive</param>
 		/// <param name="leaseId">The lease id to update</param>
+		/// <param name="reason">Reason for cancelling the lease</param>
 		/// <param name="cancellationToken">Cancellation token for the operation</param>
 		/// <returns></returns>
-		async Task CancelLeaseAsync(AgentId agentId, LeaseId leaseId, CancellationToken cancellationToken)
+		async Task CancelLeaseAsync(AgentId agentId, LeaseId leaseId, LeaseOutcomeReason reason, CancellationToken cancellationToken)
 		{
 			using TelemetrySpan span = _tracer.StartActiveSpan($"{nameof(JobService)}.{nameof(CancelLeaseAsync)}");
 			span.SetAttribute("AgentId", agentId.ToString());
 			span.SetAttribute("LeaseId", leaseId.ToString());
+			span.SetAttribute("Reason", reason.ToString());
 
 			for (; ; )
 			{
@@ -505,7 +516,7 @@ namespace HordeServer.Jobs
 					break;
 				}
 
-				IAgent? newAgent = await agent.TryCancelLeaseAsync(index, cancellationToken);
+				IAgent? newAgent = await agent.TryCancelLeaseAsync(index, reason, cancellationToken);
 				if (newAgent != null)
 				{
 					_jobTaskSource.CancelLongPollForAgent(agent.Id);
diff --git a/Engine/Source/Programs/Horde/Plugins/Compute/HordeServer.Compute/Agents/AgentCollection.cs b/Engine/Source/Programs/Horde/Plugins/Compute/HordeServer.Compute/Agents/AgentCollection.cs
--- a/Engine/Source/Programs/Horde/Plugins/Compute/HordeServer.Compute/Agents/AgentCollection.cs
+++ b/Engine/Source/Programs/Horde/Plugins/Compute/HordeServer.Compute/Agents/AgentCollection.cs
@@ -102,8 +102,8 @@ namespace HordeServer.Agents
 				=> await _collection.TryCreateLeaseAsync(this, options, cancellationToken);
 
 			/// <inheritdoc/>
-			public async Task<IAgent?> TryCancelLeaseAsync(int leaseIdx, CancellationToken cancellationToken = default)
-				=> await _collection.TryCancelLeaseAsync(this, leaseIdx, cancellationToken);
+			public async Task<IAgent?> TryCancelLeaseAsync(int leaseIdx, LeaseOutcomeReason? reason = null, CancellationToken cancellationToken = default)
+				=> await _collection.TryCancelLeaseAsync(this, leaseIdx, reason, cancellationToken);
 
 			/// <inheritdoc/>
 			public async Task<IAgent?> TryDeleteAsync(CancellationToken cancellationToken = default)
@@ -1280,7 +1280,7 @@ namespace HordeServer.Agents
 		}
 
 		/// <inheritdoc/>
-		async Task<Agent?> TryCancelLeaseAsync(Agent agent, int leaseIdx, CancellationToken cancellationToken)
+		async Task<Agent?> TryCancelLeaseAsync(Agent agent, int leaseIdx, LeaseOutcomeReason? reason, CancellationToken cancellationToken)
 		{
 			if (agent.Session == null)
 			{
@@ -1288,13 +1288,33 @@ namespace HordeServer.Agents
 			}
 
 			RpcSession? newSession = new RpcSession(agent.Session);
+			LeaseId leaseId = newSession.Leases[leaseIdx].Id;
+
 			if (newSession.Leases[leaseIdx].State == RpcLeaseState.Pending)
 			{
 				newSession.Leases.RemoveAt(leaseIdx);
+
+				// Record the outcome immediately for pending leases since they'll be removed from the session
+				if (reason.HasValue)
+				{
+					await _leaseCollection.TrySetOutcomeAsync(leaseId, _clock.UtcNow, LeaseOutcome.Cancelled, reason.Value, null, null, cancellationToken);
+				}
 			}
 			else
 			{
 				newSession.Leases[leaseIdx].State = RpcLeaseState.Cancelled;
+
+				// Pre-set the outcome for active leases - this ensures the reason is captured
+				// When the agent reports completion, TryMergeOutcomeReasonsAsync will merge additional reasons
+				if (reason.HasValue)
+				{
+					bool updated = await _leaseCollection.TrySetOutcomeAsync(leaseId, _clock.UtcNow, LeaseOutcome.Cancelled, reason.Value, null, null, cancellationToken);
+					if (!updated)
+					{
+						// Outcome was already set by another path, add our reasons to it
+						await _leaseCollection.TryAddOutcomeReasonsAsync(leaseId, reason.Value, cancellationToken);
+					}
+				}
 			}
 
 			newSession = await _scheduler.TryUpdateSessionAsync(agent.Session, newSession, null, cancellationToken);
diff --git a/Engine/Source/Programs/Horde/Plugins/Compute/HordeServer.Compute/Agents/AgentService.cs b/Engine/Source/Programs/Horde/Plugins/Compute/HordeServer.Compute/Agents/AgentService.cs
--- a/Engine/Source/Programs/Horde/Plugins/Compute/HordeServer.Compute/Agents/AgentService.cs
+++ b/Engine/Source/Programs/Horde/Plugins/Compute/HordeServer.Compute/Agents/AgentService.cs
@@ -610,7 +610,7 @@ namespace HordeServer.Agents
 			}
 
 			// Update the agent
-			await agent.TryCancelLeaseAsync(index, cancellationToken);
+			await agent.TryCancelLeaseAsync(index, outcomeReason, cancellationToken);
 			return true;
 		}
 
diff --git a/Engine/Source/Programs/Horde/Plugins/Compute/HordeServer.Compute/Agents/IAgent.cs b/Engine/Source/Programs/Horde/Plugins/Compute/HordeServer.Compute/Agents/IAgent.cs
--- a/Engine/Source/Programs/Horde/Plugins/Compute/HordeServer.Compute/Agents/IAgent.cs
+++ b/Engine/Source/Programs/Horde/Plugins/Compute/HordeServer.Compute/Agents/IAgent.cs
@@ -313,9 +313,10 @@ namespace HordeServer.Agents
 		/// Attempts to cancel a lease
 		/// </summary>
 		/// <param name="leaseIdx">Index of the lease</param>
+		/// <param name="reason">Reason for cancelling the lease</param>
 		/// <param name="cancellationToken">Cancellation token for the operation</param>
 		/// <returns>New agent state if it succeeded, otherwise null</returns>
-		Task<IAgent?> TryCancelLeaseAsync(int leaseIdx, CancellationToken cancellationToken = default);
+		Task<IAgent?> TryCancelLeaseAsync(int leaseIdx, LeaseOutcomeReason? reason = null, CancellationToken cancellationToken = default);
 
 		/// <summary>
 		/// Wait for this agent to be updated
