Horde: Add a tool tip for cancelled batch texts explaining the reasons

diff --git a/Engine/Source/Programs/Horde/HordeDashboard/src/backend/Api.ts b/Engine/Source/Programs/Horde/HordeDashboard/src/backend/Api.ts
--- a/Engine/Source/Programs/Horde/HordeDashboard/src/backend/Api.ts
+++ b/Engine/Source/Programs/Horde/HordeDashboard/src/backend/Api.ts
@@ -1576,6 +1576,9 @@ export type GetBatchResponse = {
 
 	/**Time at which the group finished (UTC) */
 	finishTime?: Date | string;
+
+	/** Human-readable names of lease outcome reasons (only populated when batch has an error) */
+	leaseOutcomeReasons?: string[];
 }
 
 /**Describes the history of a step */
diff --git a/Engine/Source/Programs/Horde/HordeDashboard/src/components/jobDetailsV2/JobDetailViewSteps.tsx b/Engine/Source/Programs/Horde/HordeDashboard/src/components/jobDetailsV2/JobDetailViewSteps.tsx
--- a/Engine/Source/Programs/Horde/HordeDashboard/src/components/jobDetailsV2/JobDetailViewSteps.tsx
+++ b/Engine/Source/Programs/Horde/HordeDashboard/src/components/jobDetailsV2/JobDetailViewSteps.tsx
@@ -274,24 +274,41 @@ const StepsPanelInner: React.FC<{ jobDetails: JobDetailsV2, depStepId?: string }
       const errorColor = failure ? (dashboard.darktheme ? "#F88070" : statusColors.get(StatusColor.Failure)) : undefined;
       const errorWeight = failure ? "600" : undefined;
 
+      // Create tooltip content from leaseOutcomeReasons
+      const leaseOutcomeReasonsTooltip = item.batch?.leaseOutcomeReasons?.length ? (
+         <Stack tokens={{ childrenGap: 4 }}>
+            <Text style={{ fontWeight: 600 }}>Cancellation reasons</Text>
+            {item.batch.leaseOutcomeReasons.map((reason, idx) => (
+               <Text key={idx}>• {reason}</Text>
+            ))}
+         </Stack>
+      ) : undefined;
+
+      // Helper to wrap element with tooltip if reasons exist
+      const wrapWithTooltip = (element: JSX.Element) => leaseOutcomeReasonsTooltip
+         ? <TooltipHost content={leaseOutcomeReasonsTooltip} id={`batch_reasons_${item.batch?.id}`} directionalHint={DirectionalHint.rightCenter}>{element}</TooltipHost>
+         : element;
+
       switch (column!.key) {
 
          case 'Name':
             if (!item.agentId) {
+               const batchTextElement = <Text styles={{ root: { color: errorColor, fontWeight: errorWeight, paddingRight: leaseOutcomeReasonsTooltip ? 4 : undefined } }} nowrap={true}>{batchText}</Text>;
                if (item.agentPool && (item.batch?.state === JobStepBatchState.Ready || item.batch?.error === JobStepBatchError.NoAgentsOnline || item.batch?.error === JobStepBatchError.NoAgentsInPool)) {
                   return <Stack horizontal disableShrink={true} horizontalAlign="start">
                      {(error || warning) && <Icon styles={{ root: { paddingTop: 3, paddingRight: 8, color: statusColors.get(StatusColor.Failure), fontWeight: errorWeight } }} iconName="Error" />}
-                     <Link to={`/pools?pool=${encodeURI(item.agentPool)}`} ><Text styles={{ root: { color: errorColor, fontWeight: errorWeight } }} nowrap={true}>{batchText}</Text></Link>
+                     <Link to={`/pools?pool=${encodeURI(item.agentPool)}`}>{wrapWithTooltip(batchTextElement)}</Link>
                   </Stack>
                }
                return <Stack horizontal disableShrink={true} horizontalAlign="start">
                   {(error || warning) && <Icon styles={{ root: { paddingTop: 3, paddingRight: 8, color: statusColors.get(StatusColor.Failure), fontWeight: errorWeight } }} iconName="Error" />}
-                  <Text styles={{ root: { color: errorColor } }} nowrap={true}>{batchText}</Text>
+                  {wrapWithTooltip(batchTextElement)}
                </Stack>;
             } else {
+               const batchTextElement = <Text styles={{ root: { cursor: 'pointer', color: errorColor, fontWeight: errorWeight, paddingRight: leaseOutcomeReasonsTooltip ? 4 : undefined } }}>{batchText}</Text>;
                return <Stack horizontal disableShrink={true}>{(error || warning) && <Icon styles={{ root: { paddingTop: 3, paddingRight: 8, color: statusColors.get(StatusColor.Failure), fontWeight: errorWeight } }} iconName="Error" />}
                   <Stack>
-                     <Link to={url} onClick={(ev) => { ev.stopPropagation(); ev.preventDefault(); setLastSelectedAgent(item.agentId); }}><Text styles={{ root: { cursor: 'pointer', color: errorColor, fontWeight: errorWeight } }} >{batchText}</Text></Link>
+                     <Link to={url} onClick={(ev) => { ev.stopPropagation(); ev.preventDefault(); setLastSelectedAgent(item.agentId); }}>{wrapWithTooltip(batchTextElement)}</Link>
                   </Stack>
                </Stack>
             }
