[Horde][BuildHealth][Dashboard] Update placeholder to handle between "New" and "Select..." labels based on admin access.

#Synced-CL 49930492

diff --git a/Engine/Source/Programs/Horde/HordeDashboard/src/components/buildhealth/BuildHealthFilterComponent.tsx b/Engine/Source/Programs/Horde/HordeDashboard/src/components/buildhealth/BuildHealthFilterComponent.tsx
--- a/Engine/Source/Programs/Horde/HordeDashboard/src/components/buildhealth/BuildHealthFilterComponent.tsx
+++ b/Engine/Source/Programs/Horde/HordeDashboard/src/components/buildhealth/BuildHealthFilterComponent.tsx
@@ -166,9 +166,11 @@ const ProjectFilterDropdownSingle: React.FC<{ handler: BuildHealthDataHandler, b
     buildHealthFilters.sort((a, b) => a.filterName.localeCompare(b.filterName));
 
     const doptions: IDropdownOption<BuildHealthFilterOptionData>[] = [];
+    const defaultSelectLabel : string = handler.hasBuildHealthFilterWriteAccess ? "(New)" : "Select a Filter";
+    
     doptions.push({
         key: "empty",
-        text: handler.hasBuildHealthFilterWriteAccess ? "(New)" : "Select a Filter",
+        text: defaultSelectLabel,
         data: undefined
     });
 
@@ -206,7 +208,7 @@ const ProjectFilterDropdownSingle: React.FC<{ handler: BuildHealthDataHandler, b
                     options={doptions}
                     styles={{ dropdown: { width: 300 } }}
                     selectedKey={computedSelectedKey}
-                    placeholder={"(New)"}
+                    placeholder={defaultSelectLabel}
                     calloutProps={{
                         directionalHint: DirectionalHint.rightTopEdge,
                         alignTargetEdge: true,
