Horde: Fix sparklines cloning lease overlays

diff --git a/Engine/Source/Programs/Horde/HordeDashboard/src/components/agents/HistoryModal.tsx b/Engine/Source/Programs/Horde/HordeDashboard/src/components/agents/HistoryModal.tsx
--- a/Engine/Source/Programs/Horde/HordeDashboard/src/components/agents/HistoryModal.tsx
+++ b/Engine/Source/Programs/Horde/HordeDashboard/src/components/agents/HistoryModal.tsx
@@ -499,8 +499,9 @@ const TelemetryPanel: React.FC<{}> = observer(({ }) => {
       const percent = chartX / sparkWidth;
 
       // Clamp and calculate sample index
+      // Use (data.length - 1) to match how processLeaseData calculates positions
       const clampedPercent = Math.max(0, Math.min(1, percent));
-      const sampleIdx = Math.min(Math.floor(clampedPercent * data.length), data.length - 1);
+      const sampleIdx = Math.min(Math.round(clampedPercent * (data.length - 1)), data.length - 1);
 
       if (sampleIdx < 0 || sampleIdx >= data.length) {
          return;
@@ -583,14 +584,18 @@ const TelemetryPanel: React.FC<{}> = observer(({ }) => {
                <TelemetryTooltip id={"cpu"} cpuKey={cpuKey} />
             </div>
 
-            <Sparklines width={sparkWidth} height={sparkHeight} data={cpuData} min={0} max={100} style={{ backgroundColor: dashboard.darktheme ? "#060709" : "#F3F2F1", padding: 8, border: "solid 1px #181A1B" }}>
-               <SparklinesLine color={dashboard.darktheme ? "lightblue" : "#1E90FF"} />
-               <SparklinesReferenceLine type="custom" value={0} style={{ stroke: statusColors.get(StatusColor.Failure) }} />
-               <SparklinesReferenceLine type="custom" value={32} style={{ stroke: statusColors.get(StatusColor.Warnings) }} />
-               <SparklinesReferenceLine type="custom" value={64} style={{ stroke: statusColors.get(StatusColor.Success) }} />
-               <SparklinesReferenceLine type="custom" value={96} style={{ stroke: statusColors.get(StatusColor.Success) }} />
-               <LeaseOverlayRects processedLeases={processedLeases} sparkHeight={sparkHeight} sparkWidth={sparkWidth} chartId="cpu" />
-            </Sparklines>
+            <div style={{ position: "relative" }}>
+               <Sparklines width={sparkWidth} height={sparkHeight} data={cpuData} min={0} max={100} style={{ backgroundColor: dashboard.darktheme ? "#060709" : "#F3F2F1", padding: 8, border: "solid 1px #181A1B" }}>
+                  <SparklinesLine color={dashboard.darktheme ? "lightblue" : "#1E90FF"} />
+                  <SparklinesReferenceLine type="custom" value={0} style={{ stroke: statusColors.get(StatusColor.Failure) }} />
+                  <SparklinesReferenceLine type="custom" value={32} style={{ stroke: statusColors.get(StatusColor.Warnings) }} />
+                  <SparklinesReferenceLine type="custom" value={64} style={{ stroke: statusColors.get(StatusColor.Success) }} />
+                  <SparklinesReferenceLine type="custom" value={96} style={{ stroke: statusColors.get(StatusColor.Success) }} />
+               </Sparklines>
+               <svg style={{ position: "absolute", top: sparkPadding + 1, left: sparkPadding + 1, pointerEvents: "none" }} width={sparkWidth} height={sparkHeight}>
+                  <LeaseOverlayRects processedLeases={processedLeases} sparkHeight={sparkHeight} sparkWidth={sparkWidth} chartId="cpu" />
+               </svg>
+            </div>
          </div>
       </Stack>
       <Stack verticalAlign="center" verticalFill>
@@ -608,14 +613,18 @@ const TelemetryPanel: React.FC<{}> = observer(({ }) => {
             <div style={{ position: "relative" }}>
                <TelemetryTooltip id="ram" />
             </div>
-            <Sparklines width={sparkWidth} height={sparkHeight} data={ramData} min={0} max={1} style={{ backgroundColor: dashboard.darktheme ? "#060709" : "#F3F2F1", padding: 8, border: "solid 1px #181A1B" }}>
-               <SparklinesLine color={dashboard.darktheme ? "lightblue" : "#1E90FF"} />
-               <SparklinesReferenceLine type="custom" value={0} style={{ stroke: statusColors.get(StatusColor.Failure) }} />
-               <SparklinesReferenceLine type="custom" value={32} style={{ stroke: statusColors.get(StatusColor.Warnings) }} />
-               <SparklinesReferenceLine type="custom" value={64} style={{ stroke: statusColors.get(StatusColor.Success) }} />
-               <SparklinesReferenceLine type="custom" value={96} style={{ stroke: statusColors.get(StatusColor.Success) }} />
-               <LeaseOverlayRects processedLeases={processedLeases} sparkHeight={sparkHeight} sparkWidth={sparkWidth} chartId="ram" />
-            </Sparklines>
+            <div style={{ position: "relative" }}>
+               <Sparklines width={sparkWidth} height={sparkHeight} data={ramData} min={0} max={1} style={{ backgroundColor: dashboard.darktheme ? "#060709" : "#F3F2F1", padding: 8, border: "solid 1px #181A1B" }}>
+                  <SparklinesLine color={dashboard.darktheme ? "lightblue" : "#1E90FF"} />
+                  <SparklinesReferenceLine type="custom" value={0} style={{ stroke: statusColors.get(StatusColor.Failure) }} />
+                  <SparklinesReferenceLine type="custom" value={32} style={{ stroke: statusColors.get(StatusColor.Warnings) }} />
+                  <SparklinesReferenceLine type="custom" value={64} style={{ stroke: statusColors.get(StatusColor.Success) }} />
+                  <SparklinesReferenceLine type="custom" value={96} style={{ stroke: statusColors.get(StatusColor.Success) }} />
+               </Sparklines>
+               <svg style={{ position: "absolute", top: sparkPadding + 1, left: sparkPadding + 1, pointerEvents: "none" }} width={sparkWidth} height={sparkHeight}>
+                  <LeaseOverlayRects processedLeases={processedLeases} sparkHeight={sparkHeight} sparkWidth={sparkWidth} chartId="ram" />
+               </svg>
+            </div>
          </div>
       </Stack>
 
@@ -634,14 +643,18 @@ const TelemetryPanel: React.FC<{}> = observer(({ }) => {
             <div style={{ position: "relative" }}>
                <TelemetryTooltip id="disk" />
             </div>
-            <Sparklines width={sparkWidth} height={sparkHeight} data={diskData} min={0} max={1} style={{ backgroundColor: dashboard.darktheme ? "#060709" : "#F3F2F1", padding: 8, border: "solid 1px #181A1B" }}>
-               <SparklinesLine color={dashboard.darktheme ? "lightblue" : "#1E90FF"} />
-               <SparklinesReferenceLine type="custom" value={0} style={{ stroke: statusColors.get(StatusColor.Failure) }} />
-               <SparklinesReferenceLine type="custom" value={32} style={{ stroke: statusColors.get(StatusColor.Warnings) }} />
-               <SparklinesReferenceLine type="custom" value={64} style={{ stroke: statusColors.get(StatusColor.Success) }} />
-               <SparklinesReferenceLine type="custom" value={96} style={{ stroke: statusColors.get(StatusColor.Success) }} />
-               <LeaseOverlayRects processedLeases={processedLeases} sparkHeight={sparkHeight} sparkWidth={sparkWidth} chartId="disk" />
-            </Sparklines>
+            <div style={{ position: "relative" }}>
+               <Sparklines width={sparkWidth} height={sparkHeight} data={diskData} min={0} max={1} style={{ backgroundColor: dashboard.darktheme ? "#060709" : "#F3F2F1", padding: 8, border: "solid 1px #181A1B" }}>
+                  <SparklinesLine color={dashboard.darktheme ? "lightblue" : "#1E90FF"} />
+                  <SparklinesReferenceLine type="custom" value={0} style={{ stroke: statusColors.get(StatusColor.Failure) }} />
+                  <SparklinesReferenceLine type="custom" value={32} style={{ stroke: statusColors.get(StatusColor.Warnings) }} />
+                  <SparklinesReferenceLine type="custom" value={64} style={{ stroke: statusColors.get(StatusColor.Success) }} />
+                  <SparklinesReferenceLine type="custom" value={96} style={{ stroke: statusColors.get(StatusColor.Success) }} />
+               </Sparklines>
+               <svg style={{ position: "absolute", top: sparkPadding + 1, left: sparkPadding + 1, pointerEvents: "none" }} width={sparkWidth} height={sparkHeight}>
+                  <LeaseOverlayRects processedLeases={processedLeases} sparkHeight={sparkHeight} sparkWidth={sparkWidth} chartId="disk" />
+               </svg>
+            </div>
          </div>
       </Stack>}
 
