Horde: Adding stream aliasing feature backend and frontend

diff --git a/Engine/Source/Programs/Horde/HordeDashboard/src/backend/Api.ts b/Engine/Source/Programs/Horde/HordeDashboard/src/backend/Api.ts
--- a/Engine/Source/Programs/Horde/HordeDashboard/src/backend/Api.ts
+++ b/Engine/Source/Programs/Horde/HordeDashboard/src/backend/Api.ts
@@ -2351,6 +2351,17 @@ export type WorkflowConfig = {
 	triageChannel?: string;
 }
 
+/**Response describing a stream alias */
+export type GetStreamAliasResponse = {
+
+	/**The unique name/identifier for this alias */
+	name: string;
+
+	/**The subset of template IDs available through this alias */
+	templates: string[];
+
+}
+
 /**Response describing a stream */
 export type GetStreamResponse = {
 
@@ -2390,6 +2401,9 @@ export type GetStreamResponse = {
 	/** Workflows for this stream */
 	workflows: WorkflowConfig[];
 
+	/** Stream aliases for this stream */
+	aliases?: GetStreamAliasResponse[];
+
 	/** Custom permissions for this object */
 	acl?: GetAclResponse;
 
diff --git a/Engine/Source/Programs/Horde/HordeDashboard/src/backend/ProjectStore.ts b/Engine/Source/Programs/Horde/HordeDashboard/src/backend/ProjectStore.ts
--- a/Engine/Source/Programs/Horde/HordeDashboard/src/backend/ProjectStore.ts
+++ b/Engine/Source/Programs/Horde/HordeDashboard/src/backend/ProjectStore.ts
@@ -88,6 +88,20 @@ export class ProjectStore {
         return stream;
     }
 
+    streamByAliasName(aliasName: string): StreamData | undefined {
+
+        let stream: StreamData | undefined;
+
+        this._projects.forEach(p => {
+            if (stream) {
+                return;
+            }
+            stream = p.streams?.find(s => s.aliases?.find(a => a.name.toLowerCase() === aliasName.toLowerCase()));
+
+        });
+        return stream;
+    }
+
 
     streamByName(project: ProjectData, streamName: string | undefined): StreamData | undefined {
         if (!project || !project.streams) {
diff --git a/Engine/Source/Programs/Horde/HordeDashboard/src/base/utilities/streamUtils.ts b/Engine/Source/Programs/Horde/HordeDashboard/src/base/utilities/streamUtils.ts
--- a/Engine/Source/Programs/Horde/HordeDashboard/src/base/utilities/streamUtils.ts
+++ b/Engine/Source/Programs/Horde/HordeDashboard/src/base/utilities/streamUtils.ts
@@ -1,5 +1,6 @@
 // Copyright Epic Games, Inc. All Rights Reserved.
 
+import { StreamData } from "../../backend/Api";
 import { projectStore } from "../../backend/ProjectStore";
 
 export const getActiveStreamId = (): string | undefined => {
@@ -19,4 +20,44 @@ export const getActiveStreamId = (): string | undefined => {
 
     return path[1];
 
+}
+
+/**
+ * Extracts the display name (last path component) from an alias name.
+ * e.g., "//Project/Stream/AliasName" -> "AliasName"
+ */
+export const getAliasDisplayName = (aliasName: string): string => {
+    const lastSlash = aliasName.lastIndexOf("/");
+    return lastSlash >= 0 ? aliasName.substring(lastSlash + 1) : aliasName;
+}
+
+/**
+ * Returns the display name for a stream in breadcrumbs.
+ * When aliasParam is present and matches a stream alias,
+ * returns "AliasDisplayName (StreamName)", otherwise returns stream.name
+ */
+export const getStreamDisplayName = (
+    stream: StreamData | undefined,
+    aliasParam: string | null
+): string => {
+    if (!stream) return "Unknown Stream";
+
+    if (aliasParam) {
+        const alias = stream.aliases?.find(a => a.name === aliasParam);
+        if (alias) {
+            return `${getAliasDisplayName(alias.name)} (${stream.name})`;
+        }
+    }
+
+    return stream.name;
+}
+
+/**
+ * Builds a stream link, preserving alias parameter if present.
+ */
+export const getStreamLink = (streamId: string, aliasParam: string | null): string => {
+    if (aliasParam) {
+        return `/stream/${streamId}?alias=${encodeURIComponent(aliasParam)}`;
+    }
+    return `/stream/${streamId}`;
 }
\ No newline at end of file
diff --git a/Engine/Source/Programs/Horde/HordeDashboard/src/components/JobView.tsx b/Engine/Source/Programs/Horde/HordeDashboard/src/components/JobView.tsx
--- a/Engine/Source/Programs/Horde/HordeDashboard/src/components/JobView.tsx
+++ b/Engine/Source/Programs/Horde/HordeDashboard/src/components/JobView.tsx
@@ -178,7 +178,7 @@ const sortJobGroups = (items: JobItem[]): [JobItem[], JobGroup[]] => {
 const commitCache = new CommitCache();
 const jobHandler = new JobHandler(false, undefined, commitCache);
 
-const JobList: React.FC<{ tab: string; filter: JobFilterSimple }> = observer(({ tab, filter }) => {
+const JobList: React.FC<{ tab: string; filter: JobFilterSimple; aliasTemplates?: string[] }> = observer(({ tab, filter, aliasTemplates }) => {
 
    const { streamId } = useParams<{ streamId: string }>();
    const query = useQuery();
@@ -222,16 +222,25 @@ const JobList: React.FC<{ tab: string; filter: JobFilterSimple }> = observer(({
       return <div />;
    }
 
+   const aliasTemplateSet = aliasTemplates ? new Set(aliasTemplates) : undefined;
+
    const templateNames: Set<string> = new Set();
    jobTab.templates?.forEach(name => {
       if (filterTemplate?.length && filterTemplate !== name) {
          return;
       }
+      // Filter by alias templates if present
+      if (aliasTemplateSet && !aliasTemplateSet.has(name)) {
+         return;
+      }
       templateNames.add(name);
    });
 
    if (!templateNames.size && filterTemplate?.length) {
-      templateNames.add(filterTemplate);
+      // Only add filterTemplate if it's in the alias templates (or no alias filter)
+      if (!aliasTemplateSet || aliasTemplateSet.has(filterTemplate)) {
+         templateNames.add(filterTemplate);
+      }
    }
 
    let jobs = jobHandler.jobs.filter(j => filterJob(j, filter.filterKeyword));
@@ -621,8 +630,8 @@ const JobList: React.FC<{ tab: string; filter: JobFilterSimple }> = observer(({
 
 });
 
-export const JobView: React.FC<{ tab: string; filter: JobFilterSimple }> = ({ tab, filter }) => {
+export const JobView: React.FC<{ tab: string; filter: JobFilterSimple; aliasTemplates?: string[] }> = ({ tab, filter, aliasTemplates }) => {
 
-   return (<JobList tab={tab} filter={filter} />);
+   return (<JobList tab={tab} filter={filter} aliasTemplates={aliasTemplates} />);
 
 };
diff --git a/Engine/Source/Programs/Horde/HordeDashboard/src/components/JobViewIncremental.tsx b/Engine/Source/Programs/Horde/HordeDashboard/src/components/JobViewIncremental.tsx
--- a/Engine/Source/Programs/Horde/HordeDashboard/src/components/JobViewIncremental.tsx
+++ b/Engine/Source/Programs/Horde/HordeDashboard/src/components/JobViewIncremental.tsx
@@ -178,7 +178,7 @@ const sortJobGroups = (items: JobItem[]): [JobItem[], JobGroup[]] => {
    return [nitems, groups];
 };
 
-const JobList: React.FC<{ tab: string; filter: JobFilterSimple }> = observer(({ tab, filter }) => {
+const JobList: React.FC<{ tab: string; filter: JobFilterSimple; aliasTemplates?: string[] }> = observer(({ tab, filter, aliasTemplates }) => {
 
    const { streamId } = useParams<{ streamId: string }>();
    const { projectStore } = useBackend();
@@ -223,8 +223,14 @@ const JobList: React.FC<{ tab: string; filter: JobFilterSimple }> = observer(({
 
    const startedByName = jobTab.showNames ? "Started By" : "Change";
 
+   const aliasTemplateSet = aliasTemplates ? new Set(aliasTemplates) : undefined;
+
    const templateNames: Set<string> = new Set();
    jobTab.templates?.forEach(name => {
+      // Filter by alias templates if present
+      if (aliasTemplateSet && !aliasTemplateSet.has(name)) {
+         return;
+      }
       templateNames.add(name);
    });
 
@@ -669,8 +675,8 @@ const JobList: React.FC<{ tab: string; filter: JobFilterSimple }> = observer(({
 
 });
 
-export const JobViewIncremental: React.FC<{ tab: string; filter: JobFilterSimple }> = ({ tab, filter }) => {
+export const JobViewIncremental: React.FC<{ tab: string; filter: JobFilterSimple; aliasTemplates?: string[] }> = ({ tab, filter, aliasTemplates }) => {
 
-   return (<JobList tab={tab} filter={filter} />);
+   return (<JobList tab={tab} filter={filter} aliasTemplates={aliasTemplates} />);
 
 };
diff --git a/Engine/Source/Programs/Horde/HordeDashboard/src/components/Preflight.tsx b/Engine/Source/Programs/Horde/HordeDashboard/src/components/Preflight.tsx
--- a/Engine/Source/Programs/Horde/HordeDashboard/src/components/Preflight.tsx
+++ b/Engine/Source/Programs/Horde/HordeDashboard/src/components/Preflight.tsx
@@ -67,6 +67,7 @@ export const PreflightRedirector: React.FC = () => {
    }
 
    let stream = projectStore.streamByFullname(streamName);
+   let resolvedAlias: string | undefined;
 
    if (!stream) {
 
@@ -76,6 +77,14 @@ export const PreflightRedirector: React.FC = () => {
          stream = projectStore.streamByFullname(streamName + "-VS");
       }
 
+      // Check if the stream name matches a stream alias
+      if (!stream) {
+         stream = projectStore.streamByAliasName(streamName);
+         if (stream) {
+            resolvedAlias = streamName;
+         }
+      }
+
    }
 
    if (!stream) {
@@ -110,27 +119,31 @@ export const PreflightRedirector: React.FC = () => {
          if (result && result.length === 1) {
 
             if (stream.id === result[0].streamId) {
-               
+
                let url = `/job/${result[0].id}?newbuild=true&allowtemplatechange=true&shelvedchange=${change}&p4v=true`;
 
                if (autosubmit === "true") {
                   url += "&autosubmit=true";
                }
-   
+
                if (templateId) {
                   url += `&templateId=${templateId}`;
                }
-   
+
                if (parameters.length) {
                   url += ("&" + parameters.join("&"));
                }
-   
+
                if (version === "2") {
                   url += "&newbuildversion=2"
                }
-   
+
+               if (resolvedAlias) {
+                  url += `&alias=${encodeURIComponent(resolvedAlias)}`;
+               }
+
                navigate(url, { replace: true });
-               return;   
+               return;
             } 
          }
 
@@ -152,6 +165,10 @@ export const PreflightRedirector: React.FC = () => {
             url += "&newbuildversion=2"
          }
 
+         if (resolvedAlias) {
+            url += `&alias=${encodeURIComponent(resolvedAlias)}`;
+         }
+
          navigate(url, { replace: true });
 
 
diff --git a/Engine/Source/Programs/Horde/HordeDashboard/src/components/ProjectHome.tsx b/Engine/Source/Programs/Horde/HordeDashboard/src/components/ProjectHome.tsx
--- a/Engine/Source/Programs/Horde/HordeDashboard/src/components/ProjectHome.tsx
+++ b/Engine/Source/Programs/Horde/HordeDashboard/src/components/ProjectHome.tsx
@@ -11,6 +11,7 @@ import backend, { useBackend } from '../backend';
 import { ProjectData, StreamData } from '../backend/Api';
 import dashboard from '../backend/Dashboard';
 import { useWindowSize } from '../base/utilities/hooks';
+import { getStreamLink } from '../base/utilities/streamUtils';
 import { getHordeTheme } from '../styles/theme';
 import { getHordeStyling } from '../styles/Styles';
 
@@ -90,12 +91,19 @@ export const ProjectHome: React.FC = observer(() => {
       crumbs.push({ text: activeProject.name });
    }
 
-   let categoryCount = localState.filteredCategories.length;
-   let width = (100 / categoryCount);
+   // Check if there are any aliases to display
+   const hasAliases = activeProject?.streams?.some(stream =>
+      stream.aliases?.some(alias =>
+         !localState.filterString || alias.name.toLowerCase().indexOf(localState.filterString.toLowerCase()) !== -1
+      )
+   );
+
+   let categoryCount = localState.filteredCategories.length + (hasAliases ? 1 : 0);
+   let width = categoryCount > 0 ? (100 / categoryCount) : 100;
 
    let categoryStackItems: any[] = [];
 
-   for (let idx = 0; idx < categoryCount; idx++) {
+   for (let idx = 0; idx < localState.filteredCategories.length; idx++) {
       let category = localState.filteredCategories[idx];
 
       let links: any[] = [];
@@ -111,6 +119,29 @@ export const ProjectHome: React.FC = observer(() => {
       );
    }
 
+   // Collect all aliases from all streams
+   const aliasLinks: any[] = [];
+   activeProject?.streams?.forEach(stream => {
+      stream.aliases?.forEach(alias => {
+         if (!localState.filterString || alias.name.toLowerCase().indexOf(localState.filterString.toLowerCase()) !== -1) {
+            aliasLinks.push(
+               <Link key={`${stream.id}_${alias.name}`} style={{ marginLeft: 10, fontSize: 15 }} to={getStreamLink(stream.id, alias.name)}>
+                  {alias.name}
+               </Link>
+            );
+         }
+      });
+   });
+
+   if (aliasLinks.length > 0) {
+      categoryStackItems.push(
+         <Stack key={'stack_aliases'} styles={{ root: { width: `${width}%` } }} tokens={{ childrenGap: 5 }}>
+            <Text variant="mediumPlus" styles={{ root: { fontFamily: "Horde Open Sans SemiBold" } }}>Aliases</Text>
+            {aliasLinks}
+         </Stack>
+      );
+   }
+
    let logoSrc = "";
    if (activeProject) {
       logoSrc = `/api/v1/projects/${activeProject.id}/logo${dashboard.darktheme ? "?darkTheme=true" : ""}`;
diff --git a/Engine/Source/Programs/Horde/HordeDashboard/src/components/StreamView.tsx b/Engine/Source/Programs/Horde/HordeDashboard/src/components/StreamView.tsx
--- a/Engine/Source/Programs/Horde/HordeDashboard/src/components/StreamView.tsx
+++ b/Engine/Source/Programs/Horde/HordeDashboard/src/components/StreamView.tsx
@@ -15,6 +15,7 @@ import { getHordeStyling } from '../styles/Styles';
 import { getHordeTheme } from '../styles/theme';
 import { BreadcrumbItem, Breadcrumbs } from './Breadcrumbs';
 import { useQuery } from "horde/base/utilities/hooks";
+import { getAliasDisplayName, getStreamDisplayName } from "../base/utilities/streamUtils";
 import { JobSearchSimpleModal } from './JobSearchSimple';
 import { JobView } from './JobView';
 import { JobViewAll } from './JobViewAll';
@@ -337,25 +338,46 @@ const StreamViewInner: React.FC = observer(() => {
    }
 
    let filter = query.get("filter") ? query.get("filter")! : undefined;
+   const aliasParam = query.get("alias");
+
+   // Get the active alias and filter tabs/templates if present
+   const activeAlias = aliasParam ? stream.aliases?.find(a => a.name === aliasParam) : undefined;
+   const aliasTemplateSet = activeAlias ? new Set(activeAlias.templates) : undefined;
+
+   // Filter tabs to only those with templates matching the alias
+   const filteredTabs = aliasTemplateSet
+      ? stream.tabs.filter(tab => {
+         const jobsTab = tab as GetJobsTabResponse;
+         return jobsTab.templates?.some(t => aliasTemplateSet.has(t));
+      })
+      : stream.tabs;
 
    let queryTab = query.get("tab") ?? undefined;
-   const currentTab = stream.tabs.find(t => t.title === queryTab) as JobsTabData | undefined;
+   const currentTab = filteredTabs.find(t => t.title === queryTab) as JobsTabData | undefined;
 
-   if (!queryTab || (queryTab.toLowerCase() !== "summary" && queryTab.toLowerCase() !== "all" && !currentTab)) {
+   // When alias is active, don't allow summary/all tabs
+   const allowSummaryAll = !activeAlias;
+   const validTab = currentTab || (allowSummaryAll && (queryTab?.toLowerCase() === "summary" || queryTab?.toLowerCase() === "all"));
 
-      if (!!stream.tabs.find(t => t.title.toLowerCase() === "incremental")) {
-         return <Navigate to={`/stream/${streamId}?tab=Incremental`} replace={true} />
-      } else if (stream.tabs.length) {
-         return <Navigate to={`/stream/${streamId}?tab=${stream.tabs[0].title}`} replace={true} />
+   if (!queryTab || !validTab) {
+      const aliasQueryParam = aliasParam ? `&alias=${encodeURIComponent(aliasParam)}` : "";
+
+      if (!!filteredTabs.find(t => t.title.toLowerCase() === "incremental")) {
+         return <Navigate to={`/stream/${streamId}?tab=Incremental${aliasQueryParam}`} replace={true} />
+      } else if (filteredTabs.length) {
+         return <Navigate to={`/stream/${streamId}?tab=${filteredTabs[0].title}${aliasQueryParam}`} replace={true} />
+      } else if (allowSummaryAll) {
+         return <Navigate to={`/stream/${streamId}?tab=all${aliasQueryParam}`} replace={true} />
       } else {
-         return <Navigate to={`/stream/${streamId}?tab=all`} replace={true} />
+         // Alias with no matching tabs - show error or redirect to stream without alias
+         return <Navigate to={`/stream/${streamId}`} replace={true} />
       }
    }
 
    const isSummary = queryTab === 'summary';
    const isAllView = queryTab === 'all';
 
-   let crumbText = stream.name;
+   let crumbText = getStreamDisplayName(stream, aliasParam);
 
    if (!isSummary && !isAllView) {
       const templateFilter = query.get("template") ?? undefined;
@@ -385,20 +407,43 @@ const StreamViewInner: React.FC = observer(() => {
       </Stack>
    }
 
-   const pivotItems = stream.tabs.map(tab => {
+   const aliasQueryParam = aliasParam ? `&alias=${encodeURIComponent(aliasParam)}` : "";
+
+   const pivotItems = filteredTabs.map(tab => {
       return <PivotItem headerText={tab.title} itemKey={tab.title} key={tab.title} onRenderItemLink={() => {
          if ((tab as GetJobsTabResponse).title === "Incremental") {
             return <HoverCard cardOpenDelay={250} type={HoverCardType.plain} plainCardProps={{ onRenderPlainCard: onRenderPlainCard, renderData: tab }}>
-               <Link to={`/stream/${streamId}?tab=${encodeURIComponent(tab.title)}`} style={{ color: modeColors.text }}>{tab.title}</Link>
+               <Link to={`/stream/${streamId}?tab=${encodeURIComponent(tab.title)}${aliasQueryParam}`} style={{ color: modeColors.text }}>{tab.title}</Link>
             </HoverCard>
          }
-         return <Link to={`/stream/${streamId}?tab=${encodeURIComponent(tab.title)}`} style={{ color: modeColors.text }}>{tab.title}</Link>
+         return <Link to={`/stream/${streamId}?tab=${encodeURIComponent(tab.title)}${aliasQueryParam}`} style={{ color: modeColors.text }}>{tab.title}</Link>
       }
       } />;
    });
 
-   pivotItems.unshift(<PivotItem headerText="All" itemKey="all" key="pivot_item_all" onRenderItemLink={() => <Link to={`/stream/${streamId}?tab=all`} style={{ color: modeColors.text }}>All</Link>} />)
-   pivotItems.unshift(<PivotItem headerText="Summary" itemKey="summary" key="pivot_item_summary" onRenderItemLink={() => <Link to={`/stream/${streamId}?tab=summary`} style={{ color: modeColors.text }}>Summary</Link>} />)
+   // Show alias indicator when viewing via alias, otherwise show All and Summary tabs
+   if (activeAlias) {
+      const aliasDisplayName = getAliasDisplayName(activeAlias.name);
+      pivotItems.unshift(
+         <PivotItem headerText={aliasDisplayName} itemKey="alias_indicator" key="pivot_item_alias" onRenderItemLink={() =>
+            <HoverCard cardOpenDelay={0} type={HoverCardType.plain} plainCardProps={{
+               onRenderPlainCard: () => (
+                  <Stack style={{ padding: 8 }}>
+                     <Link to={`/stream/${streamId}?tab=summary`} style={{ color: modeColors.text }}>View Full Stream</Link>
+                  </Stack>
+               )
+            }}>
+               <Stack horizontal verticalAlign="center" style={{ cursor: 'pointer' }}>
+                  <Text style={{ color: modeColors.text }}>{aliasDisplayName}</Text>
+                  <FontIcon iconName="ChevronDown" style={{ fontSize: 10, marginLeft: 4, color: modeColors.text }} />
+               </Stack>
+            </HoverCard>
+         } />
+      )
+   } else {
+      pivotItems.unshift(<PivotItem headerText="All" itemKey="all" key="pivot_item_all" onRenderItemLink={() => <Link to={`/stream/${streamId}?tab=all${aliasQueryParam}`} style={{ color: modeColors.text }}>All</Link>} />)
+      pivotItems.unshift(<PivotItem headerText="Summary" itemKey="summary" key="pivot_item_summary" onRenderItemLink={() => <Link to={`/stream/${streamId}?tab=summary${aliasQueryParam}`} style={{ color: modeColors.text }}>Summary</Link>} />)
+   }
 
    let findJobsItems: IContextualMenuProps = { items: [] };
 
@@ -465,7 +510,7 @@ const StreamViewInner: React.FC = observer(() => {
       <Stack className={hordeClasses.horde}>
          <TopNav />
          <Breadcrumbs items={crumbItems} title={crumbTitle} />
-         {shown && newBuildVersion == "2" && <NewBuildV2 streamId={streamId!} jobKey={newBuildTab!} show={true} onClose={(newJobId) => {
+         {shown && newBuildVersion == "2" && <NewBuildV2 streamId={streamId!} jobKey={newBuildTab!} show={true} aliasTemplates={activeAlias?.templates} onClose={(newJobId) => {
             setShown(false);
             if (newJobId) {
                navigate(`/job/${newJobId}`);
@@ -528,10 +573,10 @@ const StreamViewInner: React.FC = observer(() => {
                </Stack>
                <Stack style={{ width: "100%", paddingTop: 12, marginLeft: 4 }} >
                   <Stack >
-                     {!isSummary && !isAllView && !isIncrementalTab && queryTab && <JobView tab={queryTab} filter={simpleFilter} />}
+                     {!isSummary && !isAllView && !isIncrementalTab && queryTab && <JobView tab={queryTab} filter={simpleFilter} aliasTemplates={activeAlias?.templates} />}
                      {isSummary && <StreamSummary />}
                      {isAllView && <JobViewAll filter={simpleFilter} />}
-                     {isIncrementalTab && <JobViewIncremental tab={queryTab} filter={simpleFilter} />}
+                     {isIncrementalTab && <JobViewIncremental tab={queryTab} filter={simpleFilter} aliasTemplates={activeAlias?.templates} />}
                   </Stack>
                </Stack>
 
diff --git a/Engine/Source/Programs/Horde/HordeDashboard/src/components/TestReportView.tsx b/Engine/Source/Programs/Horde/HordeDashboard/src/components/TestReportView.tsx
--- a/Engine/Source/Programs/Horde/HordeDashboard/src/components/TestReportView.tsx
+++ b/Engine/Source/Programs/Horde/HordeDashboard/src/components/TestReportView.tsx
@@ -10,6 +10,7 @@ import { projectStore } from '../backend/ProjectStore';
 import { useParams } from 'react-router-dom';
 import { TestDataCollection, TestDataWrapper } from '../backend/TestDataHandler'
 import { useQuery } from "horde/base/utilities/hooks";
+import { getStreamDisplayName, getStreamLink } from "../base/utilities/streamUtils";
 import { getHordeStyling } from '../styles/Styles';
 
 const getFirstComponentFromPlugin = (mount: PluginMount, id: string): React.FC<any> | undefined => {
@@ -45,6 +46,7 @@ const getTestdataBreadcrumbProps = async (testdata: TestDataWrapper) => {
     const changeName = await testdata.getChangeName();
 
     const stream = getStreamData(testdata.streamId);
+    const aliasParam = new URLSearchParams(window.location.search).get("alias");
 
     const crumbItems: BreadcrumbItem[] = [
         {
@@ -52,8 +54,8 @@ const getTestdataBreadcrumbProps = async (testdata: TestDataWrapper) => {
             link: `/project/${stream?.project?.id}`
         },
         {
-            text: stream?.name ?? "Unknown Stream",
-            link: `/stream/${stream?.id}`
+            text: getStreamDisplayName(stream, aliasParam),
+            link: stream ? getStreamLink(stream.id, aliasParam) : undefined
         },
         {
             text: `${changeName}: ${jobName}`,
diff --git a/Engine/Source/Programs/Horde/HordeDashboard/src/components/TopNav.tsx b/Engine/Source/Programs/Horde/HordeDashboard/src/components/TopNav.tsx
--- a/Engine/Source/Programs/Horde/HordeDashboard/src/components/TopNav.tsx
+++ b/Engine/Source/Programs/Horde/HordeDashboard/src/components/TopNav.tsx
@@ -15,6 +15,7 @@ import { getHordeStyling } from '../styles/Styles';
 import { ArtifactQueryState, FindArtifactsModal } from './artifacts/ArtifactsSearch';
 import { getHordePlugins, getFilteredMounts, MountType } from 'hordePlugins';
 import { SupportModal } from './SupportModal';
+import { getAliasDisplayName, getStreamLink } from '../base/utilities/streamUtils';
 
 const getStyles = () => {
 
@@ -213,7 +214,6 @@ const generateProjectMenu = (store: ProjectStore) => {
             const items = useSubMenu ? catItem.subMenuProps!.items : catItem.sectionProps!.items;
 
             streams.forEach(stream => {
-
                items.push(
                   {
                      style: { color: modeColors.text, fontSize: 12 },
@@ -244,14 +244,37 @@ const generateProjectMenu = (store: ProjectStore) => {
                }
             );
          });
+      }
+
+      // Collect all aliases from all streams
+      const allAliases: IProjectContextualMenuItem[] = [];
+      p.streams?.forEach(stream => {
+         stream.aliases?.forEach(alias => {
+            allAliases.push({
+               style: { color: modeColors.text, fontSize: 12 },
+               key: `${stream.id}_alias_${alias.name}`,
+               text: getAliasDisplayName(alias.name),
+               data: { project: p, stream: stream, alias: alias },
+               link: getStreamLink(stream.id, alias.name)
+            });
+         });
+      });
 
+      if (allAliases.length > 0) {
+         subItems.push({
+            key: `aliases_${p.id}`,
+            text: "Aliases",
+            subMenuProps: {
+               contextualMenuItemAs: ProjectMenuItem,
+               items: allAliases
+            }
+         });
       }
 
       if (!subItems.length) {
          return;
       }
 
-
       const style = { ...menuStyles } as Partial<IContextualMenuStyles>;
 
       // project button
diff --git a/Engine/Source/Programs/Horde/HordeDashboard/src/components/jobDetailsV2/JobDetailViewV2.tsx b/Engine/Source/Programs/Horde/HordeDashboard/src/components/jobDetailsV2/JobDetailViewV2.tsx
--- a/Engine/Source/Programs/Horde/HordeDashboard/src/components/jobDetailsV2/JobDetailViewV2.tsx
+++ b/Engine/Source/Programs/Horde/HordeDashboard/src/components/jobDetailsV2/JobDetailViewV2.tsx
@@ -9,6 +9,7 @@ import { useWindowSize } from "../../base/utilities/hooks";
 import { getHordeStyling } from "../../styles/Styles";
 import { BreadcrumbItem, Breadcrumbs } from "../Breadcrumbs";
 import { useQuery } from "horde/base/utilities/hooks";
+import { getStreamDisplayName, getStreamLink } from "../../base/utilities/streamUtils";
 import { TopNav } from "../TopNav";
 import { JobArtifactsPanel } from "./JobDetailArtifacts";
 import { BisectionPanel } from "./JobDetailBisection";
@@ -76,13 +77,14 @@ const JobBreadCrumbs: React.FC<{ jobDetails: JobDetailsV2 }> = observer(({ jobDe
          return !!jtab.templates?.find(t => t === jobData.templateId);
       });
 
-      let streamLink = `/stream/${stream.id}`;
+      const aliasParam = query.get("alias");
+      let streamLink = getStreamLink(stream.id, aliasParam);
       if (tab) {
-         streamLink += `?tab=${tab.title}`;
+         streamLink += aliasParam ? `&tab=${tab.title}` : `?tab=${tab.title}`;
       }
 
       crumbItems.push({ text: stream.project!.name, link: `/project/${stream.project!.id}` });
-      crumbItems.push({ text: stream.name, link: streamLink });
+      crumbItems.push({ text: getStreamDisplayName(stream, aliasParam), link: streamLink });
 
    }
 
diff --git a/Engine/Source/Programs/Horde/Plugins/Build/HordeServer.Build/Streams/StreamCollection.cs b/Engine/Source/Programs/Horde/Plugins/Build/HordeServer.Build/Streams/StreamCollection.cs
--- a/Engine/Source/Programs/Horde/Plugins/Build/HordeServer.Build/Streams/StreamCollection.cs
+++ b/Engine/Source/Programs/Horde/Plugins/Build/HordeServer.Build/Streams/StreamCollection.cs
@@ -60,6 +60,10 @@ namespace HordeServer.Streams
 
 			IReadOnlyList<IWorkflow> IStream.Workflows => _config.Workflows;
 
+			IReadOnlyList<IStreamAlias>? _cachedAliases;
+			IReadOnlyList<IStreamAlias> IStream.Aliases
+				=> _cachedAliases ??= _config.Aliases.Select(x => (IStreamAlias)x).ToList();
+
 			ICommitCollection IStream.Commits => _collection._commitService.GetCollection(_config);
 
 			public Stream(StreamCollection collection, StreamConfig config, StreamDoc document)
diff --git a/Engine/Source/Programs/Horde/Plugins/Build/HordeServer.Build/Streams/StreamConfig.cs b/Engine/Source/Programs/Horde/Plugins/Build/HordeServer.Build/Streams/StreamConfig.cs
--- a/Engine/Source/Programs/Horde/Plugins/Build/HordeServer.Build/Streams/StreamConfig.cs
+++ b/Engine/Source/Programs/Horde/Plugins/Build/HordeServer.Build/Streams/StreamConfig.cs
@@ -240,6 +240,12 @@ namespace HordeServer.Streams
 		/// </summary>
 		public List<TemplateRefConfig> Templates { get; set; } = new List<TemplateRefConfig>();
 
+		/// <summary>
+		/// Stream aliases configured for this stream.
+		/// Each alias provides a named subset of the stream's templates.
+		/// </summary>
+		public List<StreamAliasConfig> Aliases { get; set; } = new List<StreamAliasConfig>();
+
 		/// <summary>
 		/// Custom permissions for this object
 		/// </summary>
@@ -952,6 +958,21 @@ namespace HordeServer.Streams
 		public string EnvVar { get; set; } = String.Empty;
 	}
 
+	/// <summary>
+	/// Configuration for a stream alias, which represents a named subset of templates
+	/// </summary>
+	public class StreamAliasConfig : IStreamAlias
+	{
+		/// <inheritdoc/>
+		[Required]
+		public string Name { get; set; } = String.Empty;
+
+		/// <inheritdoc cref="IStreamAlias.Templates"/>
+		[Required]
+		public List<TemplateId> Templates { get; set; } = new List<TemplateId>();
+		IReadOnlyList<TemplateId> IStreamAlias.Templates => Templates;
+	}
+
 	/// <summary>
 	/// Configuration for custom commit filters
 	/// </summary>
diff --git a/Engine/Source/Programs/Shared/EpicGames.Horde/Streams/IStream.cs b/Engine/Source/Programs/Shared/EpicGames.Horde/Streams/IStream.cs
--- a/Engine/Source/Programs/Shared/EpicGames.Horde/Streams/IStream.cs
+++ b/Engine/Source/Programs/Shared/EpicGames.Horde/Streams/IStream.cs
@@ -112,6 +112,12 @@ namespace EpicGames.Horde.Streams
 		/// </summary>
 		IReadOnlyList<IWorkflow> Workflows { get; }
 
+		/// <summary>
+		/// Stream aliases configured for this stream.
+		/// Each alias provides a named subset of the stream's templates.
+		/// </summary>
+		IReadOnlyList<IStreamAlias> Aliases { get; }
+
 		/// <summary>
 		/// Default settings for preflights against this stream
 		/// </summary>
@@ -388,6 +394,23 @@ namespace EpicGames.Horde.Streams
 		IChangeQuery? Change { get; }
 	}
 
+	/// <summary>
+	/// Represents a named alias for a stream with a subset of its templates
+	/// </summary>
+	public interface IStreamAlias
+	{
+		/// <summary>
+		/// The unique name/identifier for this alias
+		/// </summary>
+		string Name { get; }
+
+		/// <summary>
+		/// The subset of template IDs available through this alias.
+		/// These must be a subset of the enclosing IStream's Templates.
+		/// </summary>
+		IReadOnlyList<TemplateId> Templates { get; }
+	}
+
 	/// <summary>
 	/// Job template in a stream
 	/// </summary>
diff --git a/Engine/Source/Programs/Shared/EpicGames.Horde/Streams/StreamMessages.cs b/Engine/Source/Programs/Shared/EpicGames.Horde/Streams/StreamMessages.cs
--- a/Engine/Source/Programs/Shared/EpicGames.Horde/Streams/StreamMessages.cs
+++ b/Engine/Source/Programs/Shared/EpicGames.Horde/Streams/StreamMessages.cs
@@ -112,6 +112,11 @@ namespace EpicGames.Horde.Streams
 		/// </summary>
 		public List<GetWorkflowResponse> Workflows { get; set; }
 
+		/// <summary>
+		/// Stream aliases for this stream
+		/// </summary>
+		public List<GetStreamAliasResponse> Aliases { get; set; }
+
 		/// <summary>
 		/// Constructor
 		/// </summary>
@@ -135,11 +140,44 @@ namespace EpicGames.Horde.Streams
 			WorkspaceTypes = stream.WorkspaceTypes.ToDictionary(x => x.Key, x => new GetWorkspaceTypeResponse(x.Value));
 			Templates = templates;
 			Workflows = stream.Workflows.ConvertAll(x => new GetWorkflowResponse(x));
+			Aliases = stream.Aliases.Select(x => new GetStreamAliasResponse(x)).ToList();
 			PausedUntil = stream.PausedUntil;
 			PauseComment = stream.PauseComment;
 		}
 	}
 
+	/// <summary>
+	/// Response describing a stream alias
+	/// </summary>
+	public class GetStreamAliasResponse : IStreamAlias
+	{
+		/// <inheritdoc/>
+		public string Name { get; set; }
+
+		/// <inheritdoc cref="IStreamAlias.Templates"/>
+		public List<TemplateId> Templates { get; set; }
+		IReadOnlyList<TemplateId> IStreamAlias.Templates => Templates;
+
+		/// <summary>
+		/// Default constructor for serialization
+		/// </summary>
+		public GetStreamAliasResponse()
+		{
+			Name = String.Empty;
+			Templates = new List<TemplateId>();
+		}
+
+		/// <summary>
+		/// Constructor
+		/// </summary>
+		/// <param name="alias">The stream alias to construct from</param>
+		public GetStreamAliasResponse(IStreamAlias alias)
+		{
+			Name = alias.Name;
+			Templates = alias.Templates.ToList();
+		}
+	}
+
 	/// <summary>
 	/// Information about the default preflight to run
 	/// </summary>
