Horde: Fix stuck job/lease when retrying or aborting running steps

When retrying or aborting a running step, mark the step for abort and set the batch error to NoLongerNeeded or Cancelled. This signals that the lease should be cancelled. Also add a 24-hour timeout for stuck running steps that never complete, with periodic detection via JobExpirationService.

diff --git a/Engine/Source/Programs/Horde/Plugins/Build/HordeServer.Build.Tests/Jobs/JobCollectionTests.cs b/Engine/Source/Programs/Horde/Plugins/Build/HordeServer.Build.Tests/Jobs/JobCollectionTests.cs
--- a/Engine/Source/Programs/Horde/Plugins/Build/HordeServer.Build.Tests/Jobs/JobCollectionTests.cs
+++ b/Engine/Source/Programs/Horde/Plugins/Build/HordeServer.Build.Tests/Jobs/JobCollectionTests.cs
@@ -1144,5 +1144,93 @@ namespace HordeServer.Tests.Jobs
 		{
 			return await JobCollection.AddAsync(JobIdUtils.GenerateNewId(), _streamId, _templateId1, _template1.Hash, graph, "jobName", new CommitId(commitId), null, new CreateJobOptions());
 		}
+
+		[TestMethod]
+		public async Task RetryRunningStepShouldAbortOldStepAsync()
+		{
+			// Setup: Create a job with a single step
+			List<NewGroup> newGroups = new List<NewGroup>();
+			NewGroup group = AddGroup(newGroups);
+			AddNode(group, "Step 1", null);
+
+			IGraph graph = await GraphCollection.AppendAsync(null, newGroups, null, null);
+
+			CreateJobOptions options = new CreateJobOptions();
+			options.Arguments.Add("-Target=Step 1");
+
+			IJob job = await JobCollection.AddAsync(JobIdUtils.GenerateNewId(), new StreamId("ue4-main"), new TemplateId("test-build"), ContentHash.SHA1("hello"), graph, "Test job", CommitIdWithOrder.FromPerforceChange(123), CommitIdWithOrder.FromPerforceChange(123), options);
+
+			// Run the batch and start the step (but don't complete it)
+			job = await RunBatchAsync(job, 0);
+			job = Deref(await job.TryUpdateStepAsync(job.Batches[0].Id, job.Batches[0].Steps[0].Id, JobStepState.Running, JobStepOutcome.Success));
+
+			Assert.AreEqual(JobStepState.Running, job.Batches[0].Steps[0].State);
+
+			// Retry the RUNNING step
+			job = Deref(await job.TryUpdateStepAsync(job.Batches[0].Id, job.Batches[0].Steps[0].Id, newRetryByUserId: UserId.Anonymous));
+
+			// Verify: Old step should be marked for abort and a new batch should be created for the retry
+			Assert.AreEqual(2, job.Batches.Count, "Should have old batch + new retry batch");
+			Assert.IsTrue(job.Batches[0].Steps[0].AbortRequested, "Old running step should be aborted when retried");
+		}
+
+		[TestMethod]
+		public async Task AbortRunningStepShouldSetBatchErrorAsync()
+		{
+			// Setup: Create a job with a single step
+			List<NewGroup> newGroups = new List<NewGroup>();
+			NewGroup group = AddGroup(newGroups);
+			AddNode(group, "Step 1", null);
+
+			IGraph graph = await GraphCollection.AppendAsync(null, newGroups, null, null);
+
+			CreateJobOptions options = new CreateJobOptions();
+			options.Arguments.Add("-Target=Step 1");
+
+			IJob job = await JobCollection.AddAsync(JobIdUtils.GenerateNewId(), new StreamId("ue4-main"), new TemplateId("test-build"), ContentHash.SHA1("hello"), graph, "Test job", CommitIdWithOrder.FromPerforceChange(123), CommitIdWithOrder.FromPerforceChange(123), options);
+
+			// Run the batch and start the step
+			job = await RunBatchAsync(job, 0);
+			job = Deref(await job.TryUpdateStepAsync(job.Batches[0].Id, job.Batches[0].Steps[0].Id, JobStepState.Running, JobStepOutcome.Success));
+
+			Assert.AreEqual(JobStepState.Running, job.Batches[0].Steps[0].State);
+
+			// Abort the running step
+			job = Deref(await job.TryUpdateStepAsync(job.Batches[0].Id, job.Batches[0].Steps[0].Id, newAbortRequested: true, newAbortByUserId: UserId.Anonymous));
+
+			// Verify: Step is aborted and batch error is set
+			Assert.IsTrue(job.Batches[0].Steps[0].AbortRequested, "Step should have AbortRequested=true");
+			Assert.AreEqual(JobStepBatchError.Cancelled, job.Batches[0].Error, "Batch error should be Cancelled when running step is aborted");
+		}
+
+		[TestMethod]
+		public async Task RetryRunningStepShouldMarkOldBatchAsNoLongerNeededAsync()
+		{
+			// Setup: Create a job with a single step
+			List<NewGroup> newGroups = new List<NewGroup>();
+			NewGroup group = AddGroup(newGroups);
+			AddNode(group, "Step 1", null);
+
+			IGraph graph = await GraphCollection.AppendAsync(null, newGroups, null, null);
+
+			CreateJobOptions options = new CreateJobOptions();
+			options.Arguments.Add("-Target=Step 1");
+
+			IJob job = await JobCollection.AddAsync(JobIdUtils.GenerateNewId(), new StreamId("ue4-main"), new TemplateId("test-build"), ContentHash.SHA1("hello"), graph, "Test job", CommitIdWithOrder.FromPerforceChange(123), CommitIdWithOrder.FromPerforceChange(123), options);
+
+			// Run the batch and start the step (but don't complete it)
+			job = await RunBatchAsync(job, 0);
+			job = Deref(await job.TryUpdateStepAsync(job.Batches[0].Id, job.Batches[0].Steps[0].Id, JobStepState.Running, JobStepOutcome.Success));
+
+			Assert.AreEqual(JobStepState.Running, job.Batches[0].Steps[0].State);
+
+			// Retry the running step
+			job = Deref(await job.TryUpdateStepAsync(job.Batches[0].Id, job.Batches[0].Steps[0].Id, newRetryByUserId: UserId.Anonymous));
+
+			// Verify: Old batch should be marked NoLongerNeeded
+			Assert.AreEqual(2, job.Batches.Count, "Should have old batch + new retry batch");
+			Assert.AreEqual(JobStepBatchError.NoLongerNeeded, job.Batches[0].Error, "Old batch should be marked NoLongerNeeded when step is retried");
+			Assert.AreEqual(JobStepBatchError.None, job.Batches[1].Error, "New batch should have no error");
+		}
 	}
 }
diff --git a/Engine/Source/Programs/Horde/Plugins/Build/HordeServer.Build.Tests/Jobs/JobExpirationTests.cs b/Engine/Source/Programs/Horde/Plugins/Build/HordeServer.Build.Tests/Jobs/JobExpirationTests.cs
--- a/Engine/Source/Programs/Horde/Plugins/Build/HordeServer.Build.Tests/Jobs/JobExpirationTests.cs
+++ b/Engine/Source/Programs/Horde/Plugins/Build/HordeServer.Build.Tests/Jobs/JobExpirationTests.cs
@@ -1,15 +1,21 @@
 // Copyright Epic Games, Inc. All Rights Reserved.
 
+using EpicGames.Horde.Agents;
+using EpicGames.Horde.Agents.Leases;
+using EpicGames.Horde.Agents.Pools;
+using EpicGames.Horde.Agents.Sessions;
 using EpicGames.Horde.Commits;
 using EpicGames.Horde.Jobs;
 using EpicGames.Horde.Jobs.Graphs;
 using EpicGames.Horde.Jobs.Templates;
+using EpicGames.Horde.Logs;
 using EpicGames.Horde.Projects;
 using EpicGames.Horde.Streams;
 using HordeServer.Jobs;
 using HordeServer.Projects;
 using HordeServer.Server;
 using HordeServer.Streams;
+using HordeServer.Utilities;
 using Microsoft.Extensions.DependencyInjection;
 
 namespace HordeServer.Tests.Jobs
@@ -68,5 +74,37 @@ namespace HordeServer.Tests.Jobs
 			newJob = await JobCollection.GetAsync(job.Id);
 			Assert.IsNull(newJob);
 		}
+
+		[TestMethod]
+		public async Task StuckRunningStepShouldBeTimedOutAsync()
+		{
+			await ServiceProvider.GetRequiredService<JobExpirationService>().StartAsync(default);
+
+			IJob job = await CreateJobAsync();
+
+			// Assign a lease to the batch and start running
+			job = Deref(await job.TryAssignLeaseAsync(0, new PoolId("foo"), new AgentId("agent"), new SessionId(BinaryIdUtils.CreateNew()), new LeaseId(BinaryIdUtils.CreateNew()), new LogId(BinaryIdUtils.CreateNew())));
+			job = Deref(await job.TryUpdateBatchAsync(job.Batches[0].Id, null, JobStepBatchState.Running, null));
+
+			// Start the step (put it in running state)
+			job = Deref(await job.TryUpdateStepAsync(job.Batches[0].Id, job.Batches[0].Steps[0].Id, JobStepState.Running, JobStepOutcome.Success));
+
+			Assert.AreEqual(JobStepState.Running, job.Batches[0].Steps[0].State);
+			Assert.IsFalse(job.Batches[0].Steps[0].AbortRequested, "Step should not be aborted initially");
+
+			// Advance clock by 23 hours - step should still be running
+			await Clock.AdvanceAsync(TimeSpan.FromHours(23.0));
+
+			job = Deref(await JobCollection.GetAsync(job.Id));
+			Assert.AreEqual(JobStepState.Running, job.Batches[0].Steps[0].State);
+			Assert.IsFalse(job.Batches[0].Steps[0].AbortRequested, "Step should not be aborted before 24h");
+
+			// Advance clock past 24 hours total - step should be timed out
+			await Clock.AdvanceAsync(TimeSpan.FromHours(2.0));
+
+			job = Deref(await JobCollection.GetAsync(job.Id));
+			Assert.IsTrue(job.Batches[0].Steps[0].AbortRequested, "Step should be aborted after 24h timeout");
+			Assert.AreEqual(JobStepError.TimedOut, job.Batches[0].Steps[0].Error, "Step error should be TimedOut");
+		}
 	}
 }
diff --git a/Engine/Source/Programs/Horde/Plugins/Build/HordeServer.Build/Jobs/JobCollection.cs b/Engine/Source/Programs/Horde/Plugins/Build/HordeServer.Build/Jobs/JobCollection.cs
--- a/Engine/Source/Programs/Horde/Plugins/Build/HordeServer.Build/Jobs/JobCollection.cs
+++ b/Engine/Source/Programs/Horde/Plugins/Build/HordeServer.Build/Jobs/JobCollection.cs
@@ -1538,6 +1538,20 @@ namespace HordeServer.Jobs
 
 							step.RetriedByUserId = retryByUserId;
 							updates.Add(updateBuilder.Set(x => x.Batches[lambdaBatchIdx].Steps[lambdaStepIdx].RetriedByUserId, step.RetriedByUserId));
+
+							// If the step is currently running, abort it and mark the batch as no longer needed
+							// This ensures the agent stops executing the old step and the lease is cancelled
+							if (step is { State: JobStepState.Running, AbortRequested: false })
+							{
+								step.AbortRequested = true;
+								updates.Add(updateBuilder.Set(x => x.Batches[lambdaBatchIdx].Steps[lambdaStepIdx].AbortRequested, true));
+
+								if (batch.Error == JobStepBatchError.None)
+								{
+									batch.Error = JobStepBatchError.NoLongerNeeded;
+									updates.Add(updateBuilder.Set(x => x.Batches[lambdaBatchIdx].Error, JobStepBatchError.NoLongerNeeded));
+								}
+							}
 						}
 					}
 				}
@@ -1580,6 +1594,12 @@ namespace HordeServer.Jobs
 									newState = JobStepState.Aborted;
 									newOutcome = JobStepOutcome.Failure;
 								}
+								// If the step IS running, mark the batch for cancellation to trigger lease cancellation
+								else if (step.State == JobStepState.Running && batch.Error == JobStepBatchError.None)
+								{
+									batch.Error = JobStepBatchError.Cancelled;
+									updates.Add(updateBuilder.Set(x => x.Batches[batchIdx].Error, JobStepBatchError.Cancelled));
+								}
 
 								refreshDependentJobSteps = true;
 							}
diff --git a/Engine/Source/Programs/Horde/Plugins/Build/HordeServer.Build/Jobs/JobExpirationService.cs b/Engine/Source/Programs/Horde/Plugins/Build/HordeServer.Build/Jobs/JobExpirationService.cs
--- a/Engine/Source/Programs/Horde/Plugins/Build/HordeServer.Build/Jobs/JobExpirationService.cs
+++ b/Engine/Source/Programs/Horde/Plugins/Build/HordeServer.Build/Jobs/JobExpirationService.cs
@@ -14,6 +14,11 @@ namespace HordeServer.Jobs
 	/// </summary>
 	public sealed class JobExpirationService : IHostedService, IAsyncDisposable
 	{
+		/// <summary>
+		/// Maximum duration a step can run before being timed out
+		/// </summary>
+		public static readonly TimeSpan StepTimeout = TimeSpan.FromHours(24.0);
+
 		readonly IJobCollection _jobCollection;
 		readonly IClock _clock;
 		readonly IOptionsMonitor<BuildConfig> _buildConfig;
@@ -45,6 +50,10 @@ namespace HordeServer.Jobs
 		{
 			DateTime currentTime = _clock.UtcNow;
 
+			// Timeout stuck running steps
+			await TimeoutStuckStepsAsync(currentTime, cancellationToken);
+
+			// Expire old jobs
 			BuildConfig buildConfig = _buildConfig.CurrentValue;
 			foreach (StreamConfig streamConfig in buildConfig.Streams)
 			{
@@ -72,5 +81,45 @@ namespace HordeServer.Jobs
 				}
 			}
 		}
+
+		async Task TimeoutStuckStepsAsync(DateTime currentTime, CancellationToken cancellationToken)
+		{
+			DateTime maxStepStartTime = currentTime - StepTimeout;
+
+			// Find jobs with running batches that might have stuck steps
+			IReadOnlyList<IJob> jobs = await _jobCollection.FindAsync(
+				new FindJobOptions(BatchState: JobStepBatchState.Running),
+				count: 100,
+				cancellationToken: cancellationToken);
+
+			foreach (IJob job in jobs)
+			{
+				foreach (IJobStepBatch batch in job.Batches)
+				{
+					if (batch.State != JobStepBatchState.Running)
+					{
+						continue;
+					}
+
+					foreach (IJobStep step in batch.Steps)
+					{
+						if (step.State == JobStepState.Running &&
+							step.StartTimeUtc != null &&
+							step.StartTimeUtc.Value < maxStepStartTime &&
+							!step.AbortRequested)
+						{
+							double runningHours = (currentTime - step.StartTimeUtc.Value).TotalHours;
+							_logger.LogWarning("Timing out stuck step {JobId}:{BatchId}:{StepId} (running for {Hours:F1}h)",
+								job.Id, batch.Id, step.Id, runningHours);
+
+							await job.TryUpdateStepAsync(batch.Id, step.Id,
+								newAbortRequested: true,
+								newError: JobStepError.TimedOut,
+								cancellationToken: cancellationToken);
+						}
+					}
+				}
+			}
+		}
 	}
 }
