Horde: Fix the dashboard build

diff --git a/Engine/Source/Programs/Horde/HordeDashboard/src/components/build/NewBuildV2.tsx b/Engine/Source/Programs/Horde/HordeDashboard/src/components/build/NewBuildV2.tsx
--- a/Engine/Source/Programs/Horde/HordeDashboard/src/components/build/NewBuildV2.tsx
+++ b/Engine/Source/Programs/Horde/HordeDashboard/src/components/build/NewBuildV2.tsx
@@ -40,11 +40,12 @@ type OptionsConstruction = {
    jobDetails?: JobDetailsV2;
    jobKey?: string;
    readOnly?: boolean;
+   aliasTemplates?: string[];
 }
 
 class BuildOptions {
 
-   constructor(streamId: string, projectStore: ProjectStore, onClose: (newJobId: string | undefined) => void, jobDetails?: JobDetailsV2, jobKey?: string, readOnly?: boolean) {
+   constructor(streamId: string, projectStore: ProjectStore, onClose: (newJobId: string | undefined) => void, jobDetails?: JobDetailsV2, jobKey?: string, readOnly?: boolean, aliasTemplates?: string[]) {
 
       if (BuildOptions.instance) {
          // can happen with hot reload
@@ -60,7 +61,8 @@ class BuildOptions {
          onClose: onClose,
          jobDetails: jobDetails,
          jobKey: jobKey,
-         readOnly: readOnly
+         readOnly: readOnly,
+         aliasTemplates: aliasTemplates
       }
 
       this.init();
@@ -1047,6 +1049,13 @@ class BuildOptions {
 
       }
 
+      // Filter by alias templates if present
+      const aliasTemplates = this.construction.aliasTemplates;
+      if (aliasTemplates && aliasTemplates.length > 0) {
+         const aliasTemplateSet = new Set(aliasTemplates);
+         templates = templates.filter(t => aliasTemplateSet.has(t.id));
+      }
+
       this.templates = templates;
 
    }
@@ -2262,14 +2271,14 @@ const NewBuildV2Inner: React.FC<{ setUseLegacyDialog: (value: boolean) => void }
 
 })
 
-export const NewBuildV2: React.FC<{ streamId: string; show: boolean; onClose: (newJobId: string | undefined) => void, jobKey?: string; jobDetails?: JobDetailsV2, readOnly?: boolean }> = ({ streamId, jobKey, show, onClose, jobDetails, readOnly }) => {
+export const NewBuildV2: React.FC<{ streamId: string; show: boolean; onClose: (newJobId: string | undefined) => void, jobKey?: string; jobDetails?: JobDetailsV2, readOnly?: boolean, aliasTemplates?: string[] }> = ({ streamId, jobKey, show, onClose, jobDetails, readOnly, aliasTemplates }) => {
 
    const { projectStore } = useBackend();
    const [useNewBuildV1, setNewBuildV1] = useState<boolean>(false);
 
    // initially had this using useConst, though not working
    if (!BuildOptions.initialized()) {
-      new BuildOptions(streamId, projectStore, onClose, jobDetails, jobKey, readOnly);
+      new BuildOptions(streamId, projectStore, onClose, jobDetails, jobKey, readOnly, aliasTemplates);
    }
 
    const options = BuildOptions.get();
