Horde: Another fix to agent telemetry lease overlay

diff --git a/Engine/Source/Programs/Horde/HordeDashboard/src/components/agents/HistoryModal.tsx b/Engine/Source/Programs/Horde/HordeDashboard/src/components/agents/HistoryModal.tsx
--- a/Engine/Source/Programs/Horde/HordeDashboard/src/components/agents/HistoryModal.tsx
+++ b/Engine/Source/Programs/Horde/HordeDashboard/src/components/agents/HistoryModal.tsx
@@ -276,92 +276,59 @@ const leaseColorsLight = [
    "rgba(186, 85, 211, 0.12)",   // Medium orchid
 ];
 
-// Faint vertical line colors for lease boundaries
+// Lease boundary line colors
 const leaseBorderColorDark = "rgba(200, 200, 200, 0.3)";
 const leaseBorderColorLight = "rgba(80, 80, 80, 0.25)";
 
-const LeaseOverlayRects: React.FC<{ processedLeases: ProcessedLease[], sparkHeight: number, sparkWidth: number, chartId: string }> = ({ processedLeases, sparkHeight, sparkWidth, chartId }) => {
-
-   const colors = dashboard.darktheme ? leaseColorsDark : leaseColorsLight;
-   const borderColor = dashboard.darktheme ? leaseBorderColorDark : leaseBorderColorLight;
-
+/**
+ * Renders colored overlay rectangles and boundary lines for leases on a Sparklines chart.
+ * Each lease gets a distinct color, with vertical lines marking start/end boundaries.
+ */
+const LeaseOverlayRects: React.FC<{
+   processedLeases: ProcessedLease[];
+   sparkHeight: number;
+   sparkWidth: number;
+   chartId: string;
+}> = ({ processedLeases, sparkHeight, sparkWidth, chartId }) => {
    if (!processedLeases.length) {
       return null;
    }
 
-   const elements: JSX.Element[] = [];
+   const colors = dashboard.darktheme ? leaseColorsDark : leaseColorsLight;
+   const borderColor = dashboard.darktheme ? leaseBorderColorDark : leaseBorderColorLight;
+   const margin = 2; // Sparklines default margin
 
-   // Sort leases by start position to properly detect adjacency
+   const elements: JSX.Element[] = [];
    const sortedLeases = [...processedLeases].sort((a, b) => a.startX - b.startX);
 
    sortedLeases.forEach((pl, idx) => {
-      const color = colors[processedLeases.indexOf(pl) % colors.length];
-      let x = pl.startX;
-      let endX = pl.endX !== null ? pl.endX : sparkWidth;
-
-      // Extend to cover gaps: check if previous lease ends near this one's start
-      if (idx > 0) {
-         const prevLease = sortedLeases[idx - 1];
-         const prevEndX = prevLease.endX !== null ? prevLease.endX : sparkWidth;
-         // If there's a small gap (< 2px), extend this rect to fill it
-         if (x - prevEndX > 0 && x - prevEndX < 2) {
-            x = prevEndX;
-         }
-      }
+      const x = pl.startX;
+      const endX = pl.endX ?? sparkWidth - margin;
+      const width = endX - x;
 
-      // Extend end to next lease if there's a small gap
-      if (idx < sortedLeases.length - 1) {
-         const nextLease = sortedLeases[idx + 1];
-         // If there's a small gap (< 2px), extend this rect to fill it
-         if (nextLease.startX - endX > 0 && nextLease.startX - endX < 2) {
-            endX = nextLease.startX;
-         }
-      }
+      if (width < 1) return;
 
-      const width = endX - x;
+      const color = colors[processedLeases.indexOf(pl) % colors.length];
+
+      // Lease background rectangle
+      elements.push(
+         <rect key={`${chartId}-lease-${idx}`} x={x} y={0} width={width} height={sparkHeight} fill={color} />
+      );
 
-      if (width > 0) {
-         // Rectangle fill
+      // Start boundary line (skip if at chart start or adjacent to previous lease)
+      const prevEndX = idx > 0 ? sortedLeases[idx - 1]?.endX : null;
+      const isAdjacentToPrev = prevEndX !== null && Math.abs(x - prevEndX) < 5;
+      if (pl.startPercent > 0.01 && !isAdjacentToPrev) {
          elements.push(
-            <rect
-               key={`${chartId}-lease-${idx}`}
-               x={x}
-               y={0}
-               width={width}
-               height={sparkHeight}
-               fill={color}
-            />
+            <line key={`${chartId}-start-${idx}`} x1={x} x2={x} y1={0} y2={sparkHeight} stroke={borderColor} strokeWidth={1} />
          );
+      }
 
-         // Start line (only if not at the very beginning)
-         if (pl.startPercent > 0.01) {
-            elements.push(
-               <line
-                  key={`${chartId}-start-${idx}`}
-                  x1={pl.startX}
-                  x2={pl.startX}
-                  y1={0}
-                  y2={sparkHeight}
-                  stroke={borderColor}
-                  strokeWidth={1}
-               />
-            );
-         }
-
-         // End line (only if lease has finished and not at the very end)
-         if (pl.endX !== null && pl.endPercent !== null && pl.endPercent < 0.99) {
-            elements.push(
-               <line
-                  key={`${chartId}-end-${idx}`}
-                  x1={pl.endX}
-                  x2={pl.endX}
-                  y1={0}
-                  y2={sparkHeight}
-                  stroke={borderColor}
-                  strokeWidth={1}
-               />
-            );
-         }
+      // End boundary line (skip if ongoing or at chart end)
+      if (pl.endX !== null && pl.endPercent !== null && pl.endPercent < 0.99) {
+         elements.push(
+            <line key={`${chartId}-end-${idx}`} x1={pl.endX} x2={pl.endX} y1={0} y2={sparkHeight} stroke={borderColor} strokeWidth={1} />
+         );
       }
    });
 
@@ -490,11 +457,13 @@ const TelemetryPanel: React.FC<{}> = observer(({ }) => {
          return;
       }
 
-      let bounds = event.target.getBoundingClientRect();
+      // Use currentTarget to always get the event handler element, not whatever child is under the cursor
+      let bounds = event.currentTarget.getBoundingClientRect();
       let x = event.clientX - bounds.left;
       let y = event.clientY - bounds.top;
 
       // Adjust for padding and border to get position within the actual chart area
+      // Structure: handler div > tooltip div > styled container (border 1px, padding 8px) > inner div > SVGs
       const chartX = x - sparkPadding - 1; // -1 for border
       const percent = chartX / sparkWidth;
 
@@ -584,17 +553,19 @@ const TelemetryPanel: React.FC<{}> = observer(({ }) => {
                <TelemetryTooltip id={"cpu"} cpuKey={cpuKey} />
             </div>
 
-            <div style={{ position: "relative" }}>
-               <Sparklines width={sparkWidth} height={sparkHeight} data={cpuData} min={0} max={100} style={{ backgroundColor: dashboard.darktheme ? "#060709" : "#F3F2F1", padding: 8, border: "solid 1px #181A1B" }}>
-                  <SparklinesLine color={dashboard.darktheme ? "lightblue" : "#1E90FF"} />
-                  <SparklinesReferenceLine type="custom" value={0} style={{ stroke: statusColors.get(StatusColor.Failure) }} />
-                  <SparklinesReferenceLine type="custom" value={32} style={{ stroke: statusColors.get(StatusColor.Warnings) }} />
-                  <SparklinesReferenceLine type="custom" value={64} style={{ stroke: statusColors.get(StatusColor.Success) }} />
-                  <SparklinesReferenceLine type="custom" value={96} style={{ stroke: statusColors.get(StatusColor.Success) }} />
-               </Sparklines>
-               <svg style={{ position: "absolute", top: sparkPadding + 1, left: sparkPadding + 1, pointerEvents: "none" }} width={sparkWidth} height={sparkHeight}>
-                  <LeaseOverlayRects processedLeases={processedLeases} sparkHeight={sparkHeight} sparkWidth={sparkWidth} chartId="cpu" />
-               </svg>
+            <div style={{ backgroundColor: dashboard.darktheme ? "#060709" : "#F3F2F1", padding: 8, border: "solid 1px #181A1B" }}>
+               <div style={{ position: "relative" }}>
+                  <Sparklines width={sparkWidth} height={sparkHeight} svgWidth={sparkWidth} svgHeight={sparkHeight} data={cpuData} min={0} max={100}>
+                     <SparklinesLine color={dashboard.darktheme ? "lightblue" : "#1E90FF"} />
+                     <SparklinesReferenceLine type="custom" value={0} style={{ stroke: statusColors.get(StatusColor.Failure) }} />
+                     <SparklinesReferenceLine type="custom" value={32} style={{ stroke: statusColors.get(StatusColor.Warnings) }} />
+                     <SparklinesReferenceLine type="custom" value={64} style={{ stroke: statusColors.get(StatusColor.Success) }} />
+                     <SparklinesReferenceLine type="custom" value={96} style={{ stroke: statusColors.get(StatusColor.Success) }} />
+                  </Sparklines>
+                  <svg style={{ position: "absolute", top: 0, left: 0, pointerEvents: "none" }} width={sparkWidth} height={sparkHeight} viewBox={`0 0 ${sparkWidth} ${sparkHeight}`} preserveAspectRatio="none">
+                     <LeaseOverlayRects processedLeases={processedLeases} sparkHeight={sparkHeight} sparkWidth={sparkWidth} chartId="cpu" />
+                  </svg>
+               </div>
             </div>
          </div>
       </Stack>
@@ -613,17 +584,19 @@ const TelemetryPanel: React.FC<{}> = observer(({ }) => {
             <div style={{ position: "relative" }}>
                <TelemetryTooltip id="ram" />
             </div>
-            <div style={{ position: "relative" }}>
-               <Sparklines width={sparkWidth} height={sparkHeight} data={ramData} min={0} max={1} style={{ backgroundColor: dashboard.darktheme ? "#060709" : "#F3F2F1", padding: 8, border: "solid 1px #181A1B" }}>
-                  <SparklinesLine color={dashboard.darktheme ? "lightblue" : "#1E90FF"} />
-                  <SparklinesReferenceLine type="custom" value={0} style={{ stroke: statusColors.get(StatusColor.Failure) }} />
-                  <SparklinesReferenceLine type="custom" value={32} style={{ stroke: statusColors.get(StatusColor.Warnings) }} />
-                  <SparklinesReferenceLine type="custom" value={64} style={{ stroke: statusColors.get(StatusColor.Success) }} />
-                  <SparklinesReferenceLine type="custom" value={96} style={{ stroke: statusColors.get(StatusColor.Success) }} />
-               </Sparklines>
-               <svg style={{ position: "absolute", top: sparkPadding + 1, left: sparkPadding + 1, pointerEvents: "none" }} width={sparkWidth} height={sparkHeight}>
-                  <LeaseOverlayRects processedLeases={processedLeases} sparkHeight={sparkHeight} sparkWidth={sparkWidth} chartId="ram" />
-               </svg>
+            <div style={{ backgroundColor: dashboard.darktheme ? "#060709" : "#F3F2F1", padding: 8, border: "solid 1px #181A1B" }}>
+               <div style={{ position: "relative" }}>
+                  <Sparklines width={sparkWidth} height={sparkHeight} svgWidth={sparkWidth} svgHeight={sparkHeight} data={ramData} min={0} max={1}>
+                     <SparklinesLine color={dashboard.darktheme ? "lightblue" : "#1E90FF"} />
+                     <SparklinesReferenceLine type="custom" value={0} style={{ stroke: statusColors.get(StatusColor.Failure) }} />
+                     <SparklinesReferenceLine type="custom" value={32} style={{ stroke: statusColors.get(StatusColor.Warnings) }} />
+                     <SparklinesReferenceLine type="custom" value={64} style={{ stroke: statusColors.get(StatusColor.Success) }} />
+                     <SparklinesReferenceLine type="custom" value={96} style={{ stroke: statusColors.get(StatusColor.Success) }} />
+                  </Sparklines>
+                  <svg style={{ position: "absolute", top: 0, left: 0, pointerEvents: "none" }} width={sparkWidth} height={sparkHeight} viewBox={`0 0 ${sparkWidth} ${sparkHeight}`} preserveAspectRatio="none">
+                     <LeaseOverlayRects processedLeases={processedLeases} sparkHeight={sparkHeight} sparkWidth={sparkWidth} chartId="ram" />
+                  </svg>
+               </div>
             </div>
          </div>
       </Stack>
@@ -643,17 +616,19 @@ const TelemetryPanel: React.FC<{}> = observer(({ }) => {
             <div style={{ position: "relative" }}>
                <TelemetryTooltip id="disk" />
             </div>
-            <div style={{ position: "relative" }}>
-               <Sparklines width={sparkWidth} height={sparkHeight} data={diskData} min={0} max={1} style={{ backgroundColor: dashboard.darktheme ? "#060709" : "#F3F2F1", padding: 8, border: "solid 1px #181A1B" }}>
-                  <SparklinesLine color={dashboard.darktheme ? "lightblue" : "#1E90FF"} />
-                  <SparklinesReferenceLine type="custom" value={0} style={{ stroke: statusColors.get(StatusColor.Failure) }} />
-                  <SparklinesReferenceLine type="custom" value={32} style={{ stroke: statusColors.get(StatusColor.Warnings) }} />
-                  <SparklinesReferenceLine type="custom" value={64} style={{ stroke: statusColors.get(StatusColor.Success) }} />
-                  <SparklinesReferenceLine type="custom" value={96} style={{ stroke: statusColors.get(StatusColor.Success) }} />
-               </Sparklines>
-               <svg style={{ position: "absolute", top: sparkPadding + 1, left: sparkPadding + 1, pointerEvents: "none" }} width={sparkWidth} height={sparkHeight}>
-                  <LeaseOverlayRects processedLeases={processedLeases} sparkHeight={sparkHeight} sparkWidth={sparkWidth} chartId="disk" />
-               </svg>
+            <div style={{ backgroundColor: dashboard.darktheme ? "#060709" : "#F3F2F1", padding: 8, border: "solid 1px #181A1B" }}>
+               <div style={{ position: "relative" }}>
+                  <Sparklines width={sparkWidth} height={sparkHeight} svgWidth={sparkWidth} svgHeight={sparkHeight} data={diskData} min={0} max={1}>
+                     <SparklinesLine color={dashboard.darktheme ? "lightblue" : "#1E90FF"} />
+                     <SparklinesReferenceLine type="custom" value={0} style={{ stroke: statusColors.get(StatusColor.Failure) }} />
+                     <SparklinesReferenceLine type="custom" value={32} style={{ stroke: statusColors.get(StatusColor.Warnings) }} />
+                     <SparklinesReferenceLine type="custom" value={64} style={{ stroke: statusColors.get(StatusColor.Success) }} />
+                     <SparklinesReferenceLine type="custom" value={96} style={{ stroke: statusColors.get(StatusColor.Success) }} />
+                  </Sparklines>
+                  <svg style={{ position: "absolute", top: 0, left: 0, pointerEvents: "none" }} width={sparkWidth} height={sparkHeight} viewBox={`0 0 ${sparkWidth} ${sparkHeight}`} preserveAspectRatio="none">
+                     <LeaseOverlayRects processedLeases={processedLeases} sparkHeight={sparkHeight} sparkWidth={sparkWidth} chartId="disk" />
+                  </svg>
+               </div>
             </div>
          </div>
       </Stack>}
diff --git a/Engine/Source/Programs/Horde/HordeDashboard/src/components/agents/HistoryModalStore.ts b/Engine/Source/Programs/Horde/HordeDashboard/src/components/agents/HistoryModalStore.ts
--- a/Engine/Source/Programs/Horde/HordeDashboard/src/components/agents/HistoryModalStore.ts
+++ b/Engine/Source/Programs/Horde/HordeDashboard/src/components/agents/HistoryModalStore.ts
@@ -42,6 +42,16 @@ type InfoPanelSubItem = {
    value: string;
 }
 
+/**
+ * Processes lease data to calculate X positions for overlay rendering on Sparklines charts.
+ *
+ * Sparklines renders data points by INDEX (evenly spaced), not by time. This function maps
+ * lease time ranges to sample indices, then calculates pixel positions that align with
+ * how Sparklines positions its data points.
+ *
+ * Leases are padded by one sample width on each side to fully enclose their samples.
+ * Very short leases (those that fall between sample intervals) are filtered out.
+ */
 export function processLeaseData(
    samples: GetAgentTelemetrySampleResponse[],
    leases: GetAgentLeaseResponse[],
@@ -55,81 +65,82 @@ export function processLeaseData(
    const maxTime = samples[samples.length - 1].time;
    const minTimeMs = new Date(minTime).getTime();
    const maxTimeMs = new Date(maxTime).getTime();
-   const timeRangeMs = maxTimeMs - minTimeMs;
 
-   if (timeRangeMs <= 0) {
+   if (maxTimeMs - minTimeMs <= 0) {
       return null;
    }
 
-   const processedLeases: ProcessedLease[] = [];
-
-   // Sparklines distributes points evenly by index, so we need to find which
-   // sample index corresponds to the lease time and interpolate the position
+   // Sparklines uses margin=2, drawing points from x=2 to x=(sparkWidth-2)
+   const sparkMargin = 2;
+   const drawWidth = sparkWidth - 2 * sparkMargin;
    const numSamples = samples.length;
+   const sampleWidth = drawWidth / (numSamples - 1);
 
-   // Helper to find interpolated position for a given time
-   const findInterpolatedPosition = (targetMs: number): number => {
-      // Clamp to visible range
+   // Find first sample index on or after target time
+   const findFirstSampleOnOrAfter = (targetMs: number): number => {
       if (targetMs <= minTimeMs) return 0;
-      if (targetMs >= maxTimeMs) return 1;
-
-      // Find the two samples that bracket the target time
-      for (let i = 0; i < numSamples - 1; i++) {
-         const currentMs = new Date(samples[i].time).getTime();
-         const nextMs = new Date(samples[i + 1].time).getTime();
-
-         if (targetMs >= currentMs && targetMs <= nextMs) {
-            // Interpolate between sample i and i+1
-            const timeFraction = (targetMs - currentMs) / (nextMs - currentMs);
-            const idxFraction = (i + timeFraction) / (numSamples - 1);
-            return idxFraction;
-         }
+      if (targetMs >= maxTimeMs) return numSamples - 1;
+      for (let i = 0; i < numSamples; i++) {
+         if (new Date(samples[i].time).getTime() >= targetMs) return i;
+      }
+      return numSamples - 1;
+   };
+
+   // Find last sample index on or before target time
+   const findLastSampleOnOrBefore = (targetMs: number): number => {
+      if (targetMs <= minTimeMs) return 0;
+      if (targetMs >= maxTimeMs) return numSamples - 1;
+      for (let i = numSamples - 1; i >= 0; i--) {
+         if (new Date(samples[i].time).getTime() <= targetMs) return i;
       }
+      return 0;
+   };
 
-      // Fallback: use time-based percentage
-      return (targetMs - minTimeMs) / timeRangeMs;
+   // Convert sample index to X position
+   const indexToX = (idx: number): number => {
+      return sparkMargin + (idx / (numSamples - 1)) * drawWidth;
    };
 
+   const processedLeases: ProcessedLease[] = [];
+
    for (const lease of leases) {
       const leaseStartMs = new Date(lease.startTime).getTime();
       const leaseEndMs = lease.finishTime ? new Date(lease.finishTime).getTime() : null;
 
-      // Skip leases entirely outside the visible range
-      if (leaseEndMs !== null && leaseEndMs < minTimeMs) {
+      // Skip leases outside visible range
+      if ((leaseEndMs !== null && leaseEndMs < minTimeMs) || leaseStartMs > maxTimeMs) {
          continue;
       }
-      if (leaseStartMs > maxTimeMs) {
+
+      const startIdx = findFirstSampleOnOrAfter(leaseStartMs);
+      const endIdx = leaseEndMs !== null ? findLastSampleOnOrBefore(leaseEndMs) : null;
+
+      // Skip leases with no samples in their timeframe
+      if (endIdx !== null && startIdx > endIdx) {
          continue;
       }
 
-      // Calculate start position with interpolation
-      const startPercent = findInterpolatedPosition(leaseStartMs);
-      const startX = startPercent * sparkWidth;
+      // Calculate positions with one sample-width padding for full enclosure
+      const rawStartX = indexToX(startIdx);
+      const rawEndX = endIdx !== null ? indexToX(endIdx) : null;
+      const startX = Math.max(sparkMargin, rawStartX - sampleWidth);
+      const endX = rawEndX !== null ? Math.min(sparkMargin + drawWidth, rawEndX + sampleWidth) : null;
 
-      // Calculate end position with interpolation
-      let endPercent: number | null = null;
-      let endX: number | null = null;
-      if (leaseEndMs !== null) {
-         endPercent = findInterpolatedPosition(leaseEndMs);
-         endX = endPercent * sparkWidth;
+      // Skip leases with no visual width
+      if (endX !== null && endX <= startX) {
+         continue;
       }
 
       processedLeases.push({
          lease,
          startX,
          endX,
-         startPercent,
-         endPercent
+         startPercent: startIdx / (numSamples - 1),
+         endPercent: endIdx !== null ? endIdx / (numSamples - 1) : null
       });
    }
 
-   return {
-      samples,
-      leases,
-      processedLeases,
-      minTime,
-      maxTime
-   };
+   return { samples, leases, processedLeases, minTime, maxTime };
 }
 
 class HistoryModalStore {
