Horde: Implement dashboard plugin permission check

diff --git a/Engine/Source/Programs/Horde/HordeDashboard/plugins/index.ts b/Engine/Source/Programs/Horde/HordeDashboard/plugins/index.ts
--- a/Engine/Source/Programs/Horde/HordeDashboard/plugins/index.ts
+++ b/Engine/Source/Programs/Horde/HordeDashboard/plugins/index.ts
@@ -105,3 +105,28 @@ export function getFilteredMounts(plugin: HordePlugin): Mount[] {
       return mount.requiredAcls.every(acl => userAclSet.has(acl));
    });
 }
+
+/**
+ * Check if the current user has a specific permission for a plugin.
+ *
+ * @param pluginId - The plugin identifier (e.g., 'buildhealth')
+ * @param aclAction - The ACL action name (e.g., 'BuildHealthViewAI')
+ * @returns true if user has permission, false otherwise
+ *
+ * Semantics:
+ * - Plugin not found/not enabled → false
+ * - Plugin doesn't implement dashboard ACLs (userAcls undefined) → false
+ * - User has the specific ACL → true
+ * - Otherwise → false
+ */
+export function hasPluginPermission(pluginId: string, aclAction: string): boolean {
+   const plugin = plugins.find(p => p.id === pluginId);
+   if (!plugin?.enabled) {
+      return false;
+   }
+   // undefined userAcls means plugin doesn't implement dashboard ACLs - fail closed
+   if (plugin.userAcls === undefined) {
+      return false;
+   }
+   return plugin.userAcls.includes(aclAction);
+}
