Horde: Add ACL-gated AI Analysis link to issue modal

diff --git a/Engine/Source/Programs/Horde/HordeDashboard/src/components/IssueViewV2.tsx b/Engine/Source/Programs/Horde/HordeDashboard/src/components/IssueViewV2.tsx
--- a/Engine/Source/Programs/Horde/HordeDashboard/src/components/IssueViewV2.tsx
+++ b/Engine/Source/Programs/Horde/HordeDashboard/src/components/IssueViewV2.tsx
@@ -21,6 +21,7 @@ import { useQuery } from "horde/base/utilities/hooks";
 import { renderLine } from "./logs/LogRender";
 import { UserSelect } from "./UserSelect";
 import { CommitIdField } from "./commits/commitField";
+import { hasPluginPermission } from 'hordePlugins';
 
 const smallThreshhold = 1100;
 
@@ -1857,6 +1858,17 @@ const IssueCommandBar: React.FC = () => {
          onClick: () => { setTestFixShown(true); }
       }];
 
+   // Add AI Analysis link if user has BuildHealthViewAI permission (insert after Modify Issue)
+   if (hasPluginPermission('buildhealth', 'BuildHealthViewAI')) {
+      moreList.splice(1, 0, {
+         key: 'more_ai_analysis',
+         iconProps: { iconName: "PredictiveAnalysis" },
+         text: `AI Analysis`,
+         href: `/buildhealth/ai?issue=${issue.id}`,
+         target: "_blank",
+         style: { color: "unset" }
+      });
+   }
 
    const moreItems: ICommandBarItemProps[] = [
       {
diff --git a/Engine/Source/Programs/Horde/HordeDashboard/src/styles/Styles.tsx b/Engine/Source/Programs/Horde/HordeDashboard/src/styles/Styles.tsx
--- a/Engine/Source/Programs/Horde/HordeDashboard/src/styles/Styles.tsx
+++ b/Engine/Source/Programs/Horde/HordeDashboard/src/styles/Styles.tsx
@@ -112,7 +112,8 @@ registerIcons({
       "Info" : IconContents.INFO_SIGN,
       "Form": IconContents.FORM,
       "Grouped": IconContents.GROUP_OBJECTS,
-      "UnGrouped": IconContents.UNGROUP_OBJECTS
+      "UnGrouped": IconContents.UNGROUP_OBJECTS,
+      "PredictiveAnalysis": IconContents.PREDICTIVE_ANALYSIS
    }
 });
 
