Horde: More tweaks to the Horde Agent lease telemetry

diff --git a/Engine/Source/Programs/Horde/HordeDashboard/src/components/agents/HistoryModal.tsx b/Engine/Source/Programs/Horde/HordeDashboard/src/components/agents/HistoryModal.tsx
--- a/Engine/Source/Programs/Horde/HordeDashboard/src/components/agents/HistoryModal.tsx
+++ b/Engine/Source/Programs/Horde/HordeDashboard/src/components/agents/HistoryModal.tsx
@@ -220,6 +220,12 @@ const timeSelections: TimeSelection[] = [
    {
       text: "Past 4 Hours", key: "time_4_hours", minutes: 60 * 4
    },
+   {
+      text: "Past 8 Hours", key: "time_8_hours", minutes: 60 * 8
+   },
+   {
+      text: "Past 12 Hours", key: "time_12_hours", minutes: 60 * 12
+   },
    {
       text: "Past 1 Day", key: "time_1_day", minutes: 60 * 24
    },
@@ -290,10 +296,6 @@ const LeaseOverlayRects: React.FC<{
    sparkWidth: number;
    chartId: string;
 }> = ({ processedLeases, sparkHeight, sparkWidth, chartId }) => {
-   if (!processedLeases.length) {
-      return null;
-   }
-
    const colors = dashboard.darktheme ? leaseColorsDark : leaseColorsLight;
    const borderColor = dashboard.darktheme ? leaseBorderColorDark : leaseBorderColorLight;
    const margin = 2; // Sparklines default margin
@@ -306,7 +308,9 @@ const LeaseOverlayRects: React.FC<{
       const endX = pl.endX ?? sparkWidth - margin;
       const width = endX - x;
 
-      if (width < 1) return;
+      if (width < 1) {
+         return;
+      }
 
       const color = colors[processedLeases.indexOf(pl) % colors.length];
 
@@ -315,19 +319,24 @@ const LeaseOverlayRects: React.FC<{
          <rect key={`${chartId}-lease-${idx}`} x={x} y={0} width={width} height={sparkHeight} fill={color} />
       );
 
-      // Start boundary line (skip if at chart start or adjacent to previous lease)
+      // Start boundary line (skip if at chart start, adjacent to, or overlapping with previous lease)
       const prevEndX = idx > 0 ? sortedLeases[idx - 1]?.endX : null;
       const isAdjacentToPrev = prevEndX !== null && Math.abs(x - prevEndX) < 5;
-      if (pl.startPercent > 0.01 && !isAdjacentToPrev) {
+      const isOverlappingPrev = prevEndX !== null && x < prevEndX;
+      const drawStartLine = pl.startPercent > 0.01 && !isAdjacentToPrev && !isOverlappingPrev;
+
+      if (drawStartLine) {
          elements.push(
             <line key={`${chartId}-start-${idx}`} x1={x} x2={x} y1={0} y2={sparkHeight} stroke={borderColor} strokeWidth={1} />
          );
       }
 
       // End boundary line (skip if ongoing or at chart end)
-      if (pl.endX !== null && pl.endPercent !== null && pl.endPercent < 0.99) {
+      const drawEndLine = pl.endX !== null && pl.endPercent !== null && pl.endPercent < 0.99;
+
+      if (drawEndLine) {
          elements.push(
-            <line key={`${chartId}-end-${idx}`} x1={pl.endX} x2={pl.endX} y1={0} y2={sparkHeight} stroke={borderColor} strokeWidth={1} />
+            <line key={`${chartId}-end-${idx}`} x1={pl.endX!} x2={pl.endX!} y1={0} y2={sparkHeight} stroke={borderColor} strokeWidth={1} />
          );
       }
    });
@@ -357,39 +366,32 @@ const TelemetryPanel: React.FC<{}> = observer(({ }) => {
    if (state.currentData) { }
 
    const currentData = (state.currentData ?? []);
-
-   if (!currentData.length) {
-      return <Stack horizontalAlign="center">
-         <Spinner size={SpinnerSize.large} />
-      </Stack>;
-   }
+   const isLoading = !currentData.length;
 
    const statusColors = dashboard.getStatusColors();
 
    const sparkWidth = 1100;
    const sparkHeight = 128;
-   const data = (currentData[0] as GetAgentTelemetrySampleResponse[]);
-   const leases = ((currentData.length > 1) ? currentData[1] : []) as GetAgentLeaseResponse[];
+   const data = isLoading ? [] : (currentData[0] as GetAgentTelemetrySampleResponse[]);
+   const leases = (!isLoading && currentData.length > 1) ? currentData[1] as GetAgentLeaseResponse[] : [];
 
-   if (!data.length) {
-      return null;
-   }
+   const hasData = data.length > 0;
 
-   const processedData = processLeaseData(data, leases, sparkWidth);
+   const processedData = hasData ? processLeaseData(data, leases, sparkWidth) : null;
    const processedLeases = processedData?.processedLeases ?? [];
 
-   const cpuData = data.map(d => {
+   const cpuData = hasData ? data.map(d => {
       if (cpuKey == "user") return d.userCpu;
       if (cpuKey == "system") return d.systemCpu;
       if (cpuKey == "combined") return d.userCpu + d.systemCpu;
       return 1;
-   })
+   }) : [];
 
-   const min = data[0].time;
-   const max = data[data.length - 1].time;
+   const min = hasData ? data[0].time : "";
+   const max = hasData ? data[data.length - 1].time : "";
 
-   const minTime = getShortNiceTime(min, true);
-   const maxTime = getShortNiceTime(max, true);
+   const minTime = hasData ? getShortNiceTime(min, true) : "";
+   const maxTime = hasData ? getShortNiceTime(max, true) : "";
 
    const diskData = data.map(d => {
       if (typeof (d.totalDisk) === "number" && d.totalDisk) {
@@ -450,6 +452,7 @@ const TelemetryPanel: React.FC<{}> = observer(({ }) => {
 
    // Padding value used in Sparklines style - must match for coordinate calculations
    const sparkPadding = 8;
+   const sparkMargin = 2; // Sparklines internal margin
 
    function relativeCoords(event: any, id: string) {
 
@@ -464,8 +467,13 @@ const TelemetryPanel: React.FC<{}> = observer(({ }) => {
 
       // Adjust for padding and border to get position within the actual chart area
       // Structure: handler div > tooltip div > styled container (border 1px, padding 8px) > inner div > SVGs
-      const chartX = x - sparkPadding - 1; // -1 for border
-      const percent = chartX / sparkWidth;
+      const chartX = x - sparkPadding - 1; // -1 for border, now in SVG coordinates (0-1100)
+
+      // Sparklines draws data from x=sparkMargin to x=(sparkWidth-sparkMargin)
+      // Must account for the margin to correctly map to sample indices
+      const drawWidth = sparkWidth - 2 * sparkMargin;
+      const dataX = chartX - sparkMargin; // Position within drawable area
+      const percent = dataX / drawWidth;
 
       // Clamp and calculate sample index
       // Use (data.length - 1) to match how processLeaseData calculates positions
@@ -479,15 +487,18 @@ const TelemetryPanel: React.FC<{}> = observer(({ }) => {
       const sample = data[sampleIdx];
       const time = sample.time;
 
+      // Find lease by visual X position (using processedLeases which have padding applied)
+      // This ensures the tooltip matches what the user sees visually
+      // Sort by startX to match drawing order - leases drawn later are on top
+      // When leases overlap, use the LAST matching one (drawn on top, visually visible)
+      const sortedLeases = [...processedLeases].sort((a, b) => a.startX - b.startX);
       let lease: GetAgentLeaseResponse | undefined;
-      const timeMs = new Date(time).getTime();
-      leases.forEach(s => {
-         const startMs = new Date(s.startTime).getTime();
-         const finishMs = s.finishTime ? new Date(s.finishTime).getTime() : null;
-         if (timeMs >= startMs && (finishMs === null || timeMs <= finishMs)) {
-            lease = s;
+      for (const pl of sortedLeases) {
+         const leaseEndX = pl.endX ?? sparkWidth; // ongoing leases extend to end
+         if (chartX >= pl.startX && chartX <= leaseEndX) {
+            lease = pl.lease; // Keep updating to get the last (topmost) match
          }
-      });
+      }
 
       if (lease?.details?.type === "Job") {
          const jobId = lease.details.jobId;
@@ -507,7 +518,7 @@ const TelemetryPanel: React.FC<{}> = observer(({ }) => {
 
    return <Stack tokens={{ childrenGap: 24 }} style={{ paddingBottom: 32, paddingLeft: 8 }}>
       <Stack horizontal verticalAlign="center" verticalFill>
-         <Text variant="mediumPlus">{minTime} - {maxTime}</Text>
+         <Text variant="mediumPlus">{hasData ? `${minTime} - ${maxTime}` : "Select a time range"}</Text>
          <Stack grow />
          <Stack horizontal tokens={{ childrenGap: 18 }}>
 
@@ -537,6 +548,20 @@ const TelemetryPanel: React.FC<{}> = observer(({ }) => {
             />
          </Stack>
       </Stack>
+
+      {/* Show loading spinner when data is being fetched */}
+      {isLoading && <Stack horizontalAlign="center" style={{ paddingTop: 48 }}>
+         <Spinner size={SpinnerSize.large} label="Loading telemetry data..." />
+      </Stack>}
+
+      {/* Show message when no data is available for the selected time range */}
+      {!isLoading && !hasData && <Stack horizontalAlign="center" style={{ paddingTop: 48 }}>
+         <Text variant="large">No telemetry data available for this time range.</Text>
+         <Text variant="medium" style={{ paddingTop: 8 }}>Try selecting a different time range above.</Text>
+      </Stack>}
+
+      {/* Render charts only when data is available */}
+      {hasData && <>
       <Stack verticalAlign="center" verticalFill>
          <Stack>
             <Label>{cpuLabel} Utilization {cpuText ? `- ${cpuText}` : ""}</Label>
@@ -632,6 +657,7 @@ const TelemetryPanel: React.FC<{}> = observer(({ }) => {
             </div>
          </div>
       </Stack>}
+      </>}
 
    </Stack>
 })
diff --git a/Engine/Source/Programs/Horde/HordeDashboard/src/components/agents/HistoryModalStore.ts b/Engine/Source/Programs/Horde/HordeDashboard/src/components/agents/HistoryModalStore.ts
--- a/Engine/Source/Programs/Horde/HordeDashboard/src/components/agents/HistoryModalStore.ts
+++ b/Engine/Source/Programs/Horde/HordeDashboard/src/components/agents/HistoryModalStore.ts
@@ -101,9 +101,24 @@ export function processLeaseData(
       return sparkMargin + (idx / (numSamples - 1)) * drawWidth;
    };
 
-   const processedLeases: ProcessedLease[] = [];
+   // Sort leases by start time to properly detect adjacent leases
+   const sortedLeases = [...leases].sort((a, b) =>
+      new Date(a.startTime).getTime() - new Date(b.startTime).getTime()
+   );
+
+   // First pass: calculate raw positions for all leases
+   type IntermediateLease = {
+      lease: GetAgentLeaseResponse;
+      rawStartX: number;
+      rawEndX: number | null;
+      startX: number;
+      endX: number | null;
+      startIdx: number;
+      endIdx: number | null;
+   };
+   const intermediateLeases: IntermediateLease[] = [];
 
-   for (const lease of leases) {
+   for (const lease of sortedLeases) {
       const leaseStartMs = new Date(lease.startTime).getTime();
       const leaseEndMs = lease.finishTime ? new Date(lease.finishTime).getTime() : null;
 
@@ -120,7 +135,7 @@ export function processLeaseData(
          continue;
       }
 
-      // Calculate positions with one sample-width padding for full enclosure
+      // Calculate raw positions (without padding) and initial padded positions
       const rawStartX = indexToX(startIdx);
       const rawEndX = endIdx !== null ? indexToX(endIdx) : null;
       const startX = Math.max(sparkMargin, rawStartX - sampleWidth);
@@ -131,15 +146,61 @@ export function processLeaseData(
          continue;
       }
 
-      processedLeases.push({
+      intermediateLeases.push({
          lease,
+         rawStartX,
+         rawEndX,
          startX,
          endX,
-         startPercent: startIdx / (numSamples - 1),
-         endPercent: endIdx !== null ? endIdx / (numSamples - 1) : null
+         startIdx,
+         endIdx
       });
    }
 
+   // Second pass: adjust boundaries for adjacent leases to prevent visual overlap
+   // Adjacent leases are those where the current lease starts near where the previous ends
+   const processedLeases: ProcessedLease[] = [];
+
+   for (let i = 0; i < intermediateLeases.length; i++) {
+      const current = intermediateLeases[i];
+      const prev = i > 0 ? intermediateLeases[i - 1] : null;
+
+      let adjustedStartX = current.startX;
+      let adjustedEndX = current.endX;
+
+      // Check if this lease is adjacent to the previous one (their padded regions overlap)
+      if (prev && prev.endX !== null) {
+         // Calculate where their raw boundaries meet
+         const prevRawEnd = prev.rawEndX!;
+         const currRawStart = current.rawStartX;
+
+         // If the raw positions are close (within 2 sample widths), they're adjacent
+         // This means their padded regions would overlap
+         if (currRawStart - prevRawEnd < sampleWidth * 2) {
+            // Calculate the midpoint between the two raw positions
+            const midpoint = (prevRawEnd + currRawStart) / 2;
+
+            // Adjust the previous lease's end to the midpoint
+            if (processedLeases.length > 0) {
+               processedLeases[processedLeases.length - 1].endX = midpoint;
+            }
+
+            // Adjust current lease's start to the midpoint
+            adjustedStartX = midpoint;
+         }
+      }
+
+      const processedLease: ProcessedLease = {
+         lease: current.lease,
+         startX: adjustedStartX,
+         endX: adjustedEndX,
+         startPercent: current.startIdx / (numSamples - 1),
+         endPercent: current.endIdx !== null ? current.endIdx / (numSamples - 1) : null
+      };
+
+      processedLeases.push(processedLease);
+   }
+
    return { samples, leases, processedLeases, minTime, maxTime };
 }
 
