Horde: Remove vertical boundary lines from telemetry lease overlays and hide overlays
entirely when all leases are too narrow (<4px) to be visually meaningful.

diff --git a/Engine/Source/Programs/Horde/HordeDashboard/src/components/agents/HistoryModal.tsx b/Engine/Source/Programs/Horde/HordeDashboard/src/components/agents/HistoryModal.tsx
--- a/Engine/Source/Programs/Horde/HordeDashboard/src/components/agents/HistoryModal.tsx
+++ b/Engine/Source/Programs/Horde/HordeDashboard/src/components/agents/HistoryModal.tsx
@@ -282,13 +282,9 @@ const leaseColorsLight = [
    "rgba(186, 85, 211, 0.12)",   // Medium orchid
 ];
 
-// Lease boundary line colors
-const leaseBorderColorDark = "rgba(200, 200, 200, 0.3)";
-const leaseBorderColorLight = "rgba(80, 80, 80, 0.25)";
-
 /**
- * Renders colored overlay rectangles and boundary lines for leases on a Sparklines chart.
- * Each lease gets a distinct color, with vertical lines marking start/end boundaries.
+ * Renders colored overlay rectangles for leases on a Sparklines chart.
+ * Each lease gets a distinct color. Overlays are hidden entirely if all leases are too narrow.
  */
 const LeaseOverlayRects: React.FC<{
    processedLeases: ProcessedLease[];
@@ -297,12 +293,23 @@ const LeaseOverlayRects: React.FC<{
    chartId: string;
 }> = ({ processedLeases, sparkHeight, sparkWidth, chartId }) => {
    const colors = dashboard.darktheme ? leaseColorsDark : leaseColorsLight;
-   const borderColor = dashboard.darktheme ? leaseBorderColorDark : leaseBorderColorLight;
    const margin = 2; // Sparklines default margin
+   const minWidthThreshold = 4; // Hide all overlays if every lease is narrower than this
 
-   const elements: JSX.Element[] = [];
    const sortedLeases = [...processedLeases].sort((a, b) => a.startX - b.startX);
 
+   // Check if any lease is wide enough to show overlays
+   const hasWideEnoughLease = sortedLeases.some(pl => {
+      const endX = pl.endX ?? sparkWidth - margin;
+      return (endX - pl.startX) >= minWidthThreshold;
+   });
+
+   if (!hasWideEnoughLease) {
+      return null;
+   }
+
+   const elements: JSX.Element[] = [];
+
    sortedLeases.forEach((pl, idx) => {
       const x = pl.startX;
       const endX = pl.endX ?? sparkWidth - margin;
@@ -314,31 +321,9 @@ const LeaseOverlayRects: React.FC<{
 
       const color = colors[processedLeases.indexOf(pl) % colors.length];
 
-      // Lease background rectangle
       elements.push(
          <rect key={`${chartId}-lease-${idx}`} x={x} y={0} width={width} height={sparkHeight} fill={color} />
       );
-
-      // Start boundary line (skip if at chart start, adjacent to, or overlapping with previous lease)
-      const prevEndX = idx > 0 ? sortedLeases[idx - 1]?.endX : null;
-      const isAdjacentToPrev = prevEndX !== null && Math.abs(x - prevEndX) < 5;
-      const isOverlappingPrev = prevEndX !== null && x < prevEndX;
-      const drawStartLine = pl.startPercent > 0.01 && !isAdjacentToPrev && !isOverlappingPrev;
-
-      if (drawStartLine) {
-         elements.push(
-            <line key={`${chartId}-start-${idx}`} x1={x} x2={x} y1={0} y2={sparkHeight} stroke={borderColor} strokeWidth={1} />
-         );
-      }
-
-      // End boundary line (skip if ongoing or at chart end)
-      const drawEndLine = pl.endX !== null && pl.endPercent !== null && pl.endPercent < 0.99;
-
-      if (drawEndLine) {
-         elements.push(
-            <line key={`${chartId}-end-${idx}`} x1={pl.endX!} x2={pl.endX!} y1={0} y2={sparkHeight} stroke={borderColor} strokeWidth={1} />
-         );
-      }
    });
 
    return <>{elements}</>;
