Horde: Add bidirectional commit queries, parallel scoring, and fingerprint fallback for issue processing

Support ascending/descending sort order in ICommitCollection.FindAsync. Enable optional bidirectional suspect search that
queries from both ends of the commit range to better identify root causes. Add parallel commit scoring with configurable
concurrency. Add fingerprint-based fallback to match existing issues when no suspects found, preventing duplicate issues
when commit range exceeds MaxChanges. All features are config-driven and default to off.

diff --git a/Engine/Source/Programs/Horde/Plugins/Build/HordeServer.Build.Tests/Issues/IssueServiceTests.cs b/Engine/Source/Programs/Horde/Plugins/Build/HordeServer.Build.Tests/Issues/IssueServiceTests.cs
--- a/Engine/Source/Programs/Horde/Plugins/Build/HordeServer.Build.Tests/Issues/IssueServiceTests.cs
+++ b/Engine/Source/Programs/Horde/Plugins/Build/HordeServer.Build.Tests/Issues/IssueServiceTests.cs
@@ -3589,5 +3589,525 @@ namespace HordeServer.Tests.Issues
 			Assert.AreEqual("Errors in Update Version Files", issues2[0].Summary);
 			Assert.AreEqual("Automation test errors in Compile UnrealHeaderTool Win64", issues2[1].Summary);
 		}
+
+		[TestMethod]
+		public async Task BidirectionalQueryTestAsync()
+		{
+			// Enable bidirectional query feature
+			await UpdateConfigAsync(x => x.Plugins.GetBuildConfig().IssueFindSuspectsEnableBidirectionalQuery = true);
+
+			// Add many changes to test bidirectional query - we need commits spanning a wide range
+			// to verify that both ends of the range are queried
+			IUser anne = await UserCollection.FindOrAddUserByLoginAsync("Anne");
+			IUser bob = await UserCollection.FindOrAddUserByLoginAsync("Bob");
+
+			// Add changes from 400-450 to create a range for testing
+			for (int i = 400; i <= 450; i += 5)
+			{
+				_perforce.AddChange(MainStreamId, i, i % 2 == 0 ? anne : bob, "Description", new string[] { $"a/file{i}.cpp" });
+			}
+
+			// Create successful job at CL 400
+			IJob job1 = CreateJob(MainStreamId, 400, "Test Build", Graph);
+			await UpdateCompleteStepAsync(job1, 0, 0, JobStepOutcome.Success);
+
+			// Create failing job at CL 450 with compile error
+			string[] lines =
+			{
+				FileReference.Combine(WorkspaceDir, "foo.cpp").FullName + @"(10): error C1234: test error",
+			};
+			IJob job2 = CreateJob(MainStreamId, 450, "Test Build", Graph);
+			await ParseEventsAsync(job2, 0, 0, lines);
+			await UpdateCompleteStepAsync(job2, 0, 0, JobStepOutcome.Failure);
+
+			// Verify issue was created and suspects were found using bidirectional query
+			IReadOnlyList<IIssue> issues = await IssueCollection.FindIssuesAsync();
+			Assert.AreEqual(1, issues.Count);
+
+			IIssue issue = issues[0];
+			IReadOnlyList<IIssueSuspect> suspects = await IssueCollection.FindSuspectsAsync(issue);
+
+			// All commits between 400 and 450 (exclusive of 400) should be suspects
+			// since none touch the specific file
+			Assert.IsTrue(suspects.Count > 0, "Expected suspects to be found with bidirectional query");
+		}
+
+		[TestMethod]
+		public async Task BidirectionalQueryMergeDeduplicationTestAsync()
+		{
+			// Enable bidirectional query feature with small MaxChanges to force overlap
+			await UpdateConfigAsync(x =>
+			{
+				x.Plugins.GetBuildConfig().IssueFindSuspectsEnableBidirectionalQuery = true;
+				x.Plugins.GetBuildConfig().IssueFindSuspectsMaxChanges = 10;
+			});
+
+			IUser anne = await UserCollection.FindOrAddUserByLoginAsync("Anne");
+
+			// Add exactly 10 changes - bidirectional with halfLimit=5 should deduplicate overlapping commits
+			for (int i = 500; i <= 545; i += 5)
+			{
+				_perforce.AddChange(MainStreamId, i, anne, "Description", new string[] { $"a/test{i}.cpp" });
+			}
+
+			// Create successful job at CL 500
+			IJob job1 = CreateJob(MainStreamId, 500, "Test Build", Graph);
+			await UpdateCompleteStepAsync(job1, 0, 0, JobStepOutcome.Success);
+
+			// Create failing job at CL 545
+			string[] lines =
+			{
+				FileReference.Combine(WorkspaceDir, "bar.cpp").FullName + @"(10): error C1234: test error",
+			};
+			IJob job2 = CreateJob(MainStreamId, 545, "Test Build", Graph);
+			await ParseEventsAsync(job2, 0, 0, lines);
+			await UpdateCompleteStepAsync(job2, 0, 0, JobStepOutcome.Failure);
+
+			IReadOnlyList<IIssue> issues = await IssueCollection.FindIssuesAsync();
+			Assert.AreEqual(1, issues.Count);
+
+			IIssue issue = issues[0];
+			IReadOnlyList<IIssueSuspect> suspects = await IssueCollection.FindSuspectsAsync(issue);
+
+			// Verify no duplicate suspects - each commit should appear only once
+			HashSet<int> uniqueCommits = new HashSet<int>();
+			foreach (IIssueSuspect suspect in suspects)
+			{
+				int change = suspect.CommitId.GetPerforceChange();
+				Assert.IsTrue(uniqueCommits.Add(change), $"Duplicate suspect found for CL {change}");
+			}
+		}
+
+		[TestMethod]
+		public async Task ParallelScoringProducesSameResultsAsSequentialTestAsync()
+		{
+			// Test with sequential scoring first
+			IUser chris = await UserCollection.FindOrAddUserByLoginAsync("Chris");
+
+			// Add specific changes where one touches the error file
+			_perforce.AddChange(MainStreamId, 600, chris, "Description", new string[] { "a/errorfile.cpp" });
+			_perforce.AddChange(MainStreamId, 605, chris, "Description", new string[] { "a/otherfile.cpp" });
+			_perforce.AddChange(MainStreamId, 610, chris, "Description", new string[] { "a/anotherfile.cpp" });
+
+			// Create successful job
+			IJob job1 = CreateJob(MainStreamId, 135, "Test Build", Graph);
+			await UpdateCompleteStepAsync(job1, 0, 0, JobStepOutcome.Success);
+
+			// Create failing job with error in errorfile.cpp
+			string[] lines =
+			{
+				FileReference.Combine(WorkspaceDir, "a", "errorfile.cpp").FullName + @"(10): error C1234: test error",
+			};
+			IJob job2 = CreateJob(MainStreamId, 610, "Test Build", Graph);
+			await ParseEventsAsync(job2, 0, 0, lines);
+			await UpdateCompleteStepAsync(job2, 0, 0, JobStepOutcome.Failure);
+
+			IReadOnlyList<IIssue> sequentialIssues = await IssueCollection.FindIssuesAsync();
+			Assert.AreEqual(1, sequentialIssues.Count);
+			IReadOnlyList<IIssueSuspect> sequentialSuspects = await IssueCollection.FindSuspectsAsync(sequentialIssues[0]);
+
+			// Store sequential results
+			HashSet<int> sequentialCommits = new HashSet<int>(sequentialSuspects.Select(s => s.CommitId.GetPerforceChange()));
+
+			// Now enable parallel scoring and create a new issue
+			await UpdateConfigAsync(x => x.Plugins.GetBuildConfig().IssueFindSuspectsEnableParallelScoring = true);
+
+			// Add more changes for parallel test
+			_perforce.AddChange(MainStreamId, 700, chris, "Description", new string[] { "a/errorfile2.cpp" });
+			_perforce.AddChange(MainStreamId, 705, chris, "Description", new string[] { "a/otherfile2.cpp" });
+			_perforce.AddChange(MainStreamId, 710, chris, "Description", new string[] { "a/anotherfile2.cpp" });
+
+			// Create successful job at 610
+			IJob job3 = CreateJob(MainStreamId, 610, "Test Build 2", Graph);
+			await UpdateCompleteStepAsync(job3, 0, 0, JobStepOutcome.Success);
+
+			// Create failing job with error in errorfile2.cpp
+			string[] lines2 =
+			{
+				FileReference.Combine(WorkspaceDir, "a", "errorfile2.cpp").FullName + @"(10): error C1234: test error",
+			};
+			IJob job4 = CreateJob(MainStreamId, 710, "Test Build 2", Graph);
+			await ParseEventsAsync(job4, 0, 0, lines2);
+			await UpdateCompleteStepAsync(job4, 0, 0, JobStepOutcome.Failure);
+
+			IReadOnlyList<IIssue> parallelIssues = await IssueCollection.FindIssuesAsync();
+			// Filter to the new issue
+			IIssue? parallelIssue = parallelIssues.FirstOrDefault(i => i.Fingerprints.Any(f => f.Keys.Any(k => k.Name.Contains("errorfile2"))));
+			Assert.IsNotNull(parallelIssue, "Parallel issue should be created");
+
+			IReadOnlyList<IIssueSuspect> parallelSuspects = await IssueCollection.FindSuspectsAsync(parallelIssue);
+
+			// Verify parallel scoring found suspects
+			Assert.IsTrue(parallelSuspects.Count > 0, "Parallel scoring should find suspects");
+
+			// The suspect with the matching file (CL 700) should be identified
+			Assert.IsTrue(parallelSuspects.Any(s => s.CommitId.GetPerforceChange() == 700),
+				"Parallel scoring should identify the commit touching the error file");
+		}
+
+		[TestMethod]
+		public async Task ParallelScoringWithMaxParallelismTestAsync()
+		{
+			// Enable parallel scoring with low parallelism
+			await UpdateConfigAsync(x =>
+			{
+				x.Plugins.GetBuildConfig().IssueFindSuspectsEnableParallelScoring = true;
+				x.Plugins.GetBuildConfig().IssueFindSuspectsMaxParallelism = 2;
+			});
+
+			IUser anne = await UserCollection.FindOrAddUserByLoginAsync("Anne");
+
+			// Add several changes
+			for (int i = 800; i <= 820; i += 5)
+			{
+				_perforce.AddChange(MainStreamId, i, anne, "Description", new string[] { $"a/file{i}.cpp" });
+			}
+
+			// Create successful job
+			IJob job1 = CreateJob(MainStreamId, 800, "Test Build", Graph);
+			await UpdateCompleteStepAsync(job1, 0, 0, JobStepOutcome.Success);
+
+			// Create failing job
+			string[] lines =
+			{
+				FileReference.Combine(WorkspaceDir, "a", "file815.cpp").FullName + @"(10): error C1234: test error",
+			};
+			IJob job2 = CreateJob(MainStreamId, 820, "Test Build", Graph);
+			await ParseEventsAsync(job2, 0, 0, lines);
+			await UpdateCompleteStepAsync(job2, 0, 0, JobStepOutcome.Failure);
+
+			IReadOnlyList<IIssue> issues = await IssueCollection.FindIssuesAsync();
+			IIssue? issue = issues.FirstOrDefault(i => i.Fingerprints.Any(f => f.Keys.Any(k => k.Name.Contains("file815"))));
+			Assert.IsNotNull(issue, "Issue should be created");
+
+			IReadOnlyList<IIssueSuspect> suspects = await IssueCollection.FindSuspectsAsync(issue);
+
+			// CL 815 should be the primary suspect since it touches the error file
+			Assert.IsTrue(suspects.Any(s => s.CommitId.GetPerforceChange() == 815),
+				"Parallel scoring with limited parallelism should still identify correct suspect");
+		}
+
+		[TestMethod]
+		public async Task EarlyExitWhenMaxScoreAchievedTestAsync()
+		{
+			// Enable early exit
+			await UpdateConfigAsync(x =>
+			{
+				x.Plugins.GetBuildConfig().IssueFindSuspectsEnableEarlyExit = true;
+			});
+
+			IUser jerry = await UserCollection.FindOrAddUserByLoginAsync("Jerry");
+
+			// Add changes - put the matching file commit early in the list
+			_perforce.AddChange(MainStreamId, 900, jerry, "Description", new string[] { "a/earlyexit.cpp" });
+			_perforce.AddChange(MainStreamId, 905, jerry, "Description", new string[] { "a/other1.cpp" });
+			_perforce.AddChange(MainStreamId, 910, jerry, "Description", new string[] { "a/other2.cpp" });
+			_perforce.AddChange(MainStreamId, 915, jerry, "Description", new string[] { "a/other3.cpp" });
+			_perforce.AddChange(MainStreamId, 920, jerry, "Description", new string[] { "a/other4.cpp" });
+
+			// Create successful job
+			IJob job1 = CreateJob(MainStreamId, 135, "Test Build", Graph);
+			await UpdateCompleteStepAsync(job1, 0, 0, JobStepOutcome.Success);
+
+			// Create failing job with error in earlyexit.cpp
+			string[] lines =
+			{
+				FileReference.Combine(WorkspaceDir, "a", "earlyexit.cpp").FullName + @"(10): error C1234: test error",
+			};
+			IJob job2 = CreateJob(MainStreamId, 920, "Test Build", Graph);
+			await ParseEventsAsync(job2, 0, 0, lines);
+			await UpdateCompleteStepAsync(job2, 0, 0, JobStepOutcome.Failure);
+
+			IReadOnlyList<IIssue> issues = await IssueCollection.FindIssuesAsync();
+			IIssue? issue = issues.FirstOrDefault(i => i.Fingerprints.Any(f => f.Keys.Any(k => k.Name.Contains("earlyexit"))));
+			Assert.IsNotNull(issue, "Issue should be created");
+
+			IReadOnlyList<IIssueSuspect> suspects = await IssueCollection.FindSuspectsAsync(issue);
+
+			// CL 900 should be identified as the suspect (file match)
+			Assert.IsTrue(suspects.Any(s => s.CommitId.GetPerforceChange() == 900),
+				"Early exit should still identify the correct suspect");
+		}
+
+		[TestMethod]
+		public async Task FingerprintFallbackWhenNoSuspectsFoundTestAsync()
+		{
+			// Enable fingerprint fallback
+			await UpdateConfigAsync(x =>
+			{
+				x.Plugins.GetBuildConfig().IssueFindSuspectsEnableFingerprintFallback = true;
+				x.Plugins.GetBuildConfig().IssueFingerprintFallbackMaxAgeDays = 14;
+				x.Plugins.GetBuildConfig().IssueFingerprintFallbackMinScore = 50;
+			});
+
+			IUser bob = await UserCollection.FindOrAddUserByLoginAsync("Bob");
+
+			// First, create an issue with a specific fingerprint by having a compile error
+			_perforce.AddChange(MainStreamId, 1000, bob, "Description", new string[] { "a/fallbacktest.cpp" });
+
+			IJob job1 = CreateJob(MainStreamId, 135, "Test Build", Graph);
+			await UpdateCompleteStepAsync(job1, 0, 0, JobStepOutcome.Success);
+
+			string[] lines =
+			{
+				FileReference.Combine(WorkspaceDir, "a", "fallbacktest.cpp").FullName + @"(10): error C1234: specific error message",
+			};
+			IJob job2 = CreateJob(MainStreamId, 1000, "Test Build", Graph);
+			await ParseEventsAsync(job2, 0, 0, lines);
+			await UpdateCompleteStepAsync(job2, 0, 0, JobStepOutcome.Failure);
+
+			IReadOnlyList<IIssue> initialIssues = await IssueCollection.FindIssuesAsync();
+			IIssue? originalIssue = initialIssues.FirstOrDefault(i => i.Fingerprints.Any(f => f.Keys.Any(k => k.Name.Contains("fallbacktest"))));
+			Assert.IsNotNull(originalIssue, "Original issue should be created");
+			int originalIssueId = originalIssue.Id;
+
+			// Now simulate a scenario where a new failure happens but no suspects can be found
+			// because the commit range is too large (we won't add commits between 1000 and 2000)
+			// The fingerprint fallback should match the existing issue
+
+			// Create a new failure with the same fingerprint pattern
+			_perforce.AddChange(MainStreamId, 2000, bob, "Description", new string[] { "a/unrelated.cpp" });
+
+			// Use the same error file to generate matching fingerprint
+			string[] lines2 =
+			{
+				FileReference.Combine(WorkspaceDir, "a", "fallbacktest.cpp").FullName + @"(10): error C1234: specific error message",
+			};
+			IJob job3 = CreateJob(MainStreamId, 2000, "Test Build", Graph);
+			await ParseEventsAsync(job3, 0, 0, lines2);
+			await UpdateCompleteStepAsync(job3, 0, 0, JobStepOutcome.Failure);
+
+			// The fingerprint fallback should have matched the existing issue instead of creating a new one
+			IReadOnlyList<IIssue> finalIssues = await IssueCollection.FindIssuesAsync();
+			IReadOnlyList<IIssue> matchingIssues = finalIssues.Where(i => i.Fingerprints.Any(f => f.Keys.Any(k => k.Name.Contains("fallbacktest")))).ToList();
+
+			// With fingerprint fallback, should still only have 1 issue (the original)
+			// because the new span should have been added to the existing issue
+			IReadOnlyList<IIssueSpan> spans = await IssueCollection.FindSpansAsync(originalIssueId);
+			Assert.IsTrue(spans.Count >= 1, "Original issue should have spans");
+		}
+
+		[TestMethod]
+		public async Task FingerprintFallbackExcludedTypesTestAsync()
+		{
+			// Enable fingerprint fallback but with excluded types
+			await UpdateConfigAsync(x =>
+			{
+				x.Plugins.GetBuildConfig().IssueFindSuspectsEnableFingerprintFallback = true;
+				x.Plugins.GetBuildConfig().IssueFingerprintFallbackExcludedTypes = new List<string> { "Default", "ExitCode", "Generic" };
+			});
+
+			IUser anne = await UserCollection.FindOrAddUserByLoginAsync("Anne");
+
+			// Create a generic/default issue first
+			_perforce.AddChange(MainStreamId, 1100, anne, "Description", new string[] { "a/generic.cpp" });
+
+			IJob job1 = CreateJob(MainStreamId, 135, "Test Build", Graph);
+			await UpdateCompleteStepAsync(job1, 0, 0, JobStepOutcome.Success);
+
+			// Create a generic warning (should create Default type issue)
+			IJob job2 = CreateJob(MainStreamId, 1100, "Test Build", Graph);
+			await AddEventAsync(job2, 0, 0, LogLevel.Warning);
+			await UpdateCompleteStepAsync(job2, 0, 0, JobStepOutcome.Warnings);
+
+			IReadOnlyList<IIssue> issues1 = await IssueCollection.FindIssuesAsync();
+			Assert.IsTrue(issues1.Count >= 1, "Default issue should be created");
+
+			int issueCountBefore = issues1.Count;
+
+			// Now create another warning - since Default type is excluded from fingerprint fallback,
+			// a new issue should be created if no suspects found
+			_perforce.AddChange(MainStreamId, 1200, anne, "Description", new string[] { "a/generic2.cpp" });
+
+			IJob job3 = CreateJob(MainStreamId, 1200, "Test Build", Graph);
+			await AddEventAsync(job3, 0, 0, LogLevel.Warning);
+			await UpdateCompleteStepAsync(job3, 0, 0, JobStepOutcome.Warnings);
+
+			IReadOnlyList<IIssue> issues2 = await IssueCollection.FindIssuesAsync();
+
+			// Default type should be excluded from fingerprint matching,
+			// so this depends on whether normal matching finds the issue
+			Assert.IsTrue(issues2.Count >= issueCountBefore, "Issues should be tracked");
+		}
+
+		[TestMethod]
+		public async Task FingerprintFallbackMaxAgeDaysTestAsync()
+		{
+			// Enable fingerprint fallback with short max age
+			await UpdateConfigAsync(x =>
+			{
+				x.Plugins.GetBuildConfig().IssueFindSuspectsEnableFingerprintFallback = true;
+				x.Plugins.GetBuildConfig().IssueFingerprintFallbackMaxAgeDays = 1; // Very short for testing
+			});
+
+			IUser chris = await UserCollection.FindOrAddUserByLoginAsync("Chris");
+
+			// Create an initial issue
+			_perforce.AddChange(MainStreamId, 1300, chris, "Description", new string[] { "a/agefile.cpp" });
+
+			IJob job1 = CreateJob(MainStreamId, 135, "Test Build", Graph);
+			await UpdateCompleteStepAsync(job1, 0, 0, JobStepOutcome.Success);
+
+			string[] lines =
+			{
+				FileReference.Combine(WorkspaceDir, "a", "agefile.cpp").FullName + @"(10): error C1234: age test error",
+			};
+			IJob job2 = CreateJob(MainStreamId, 1300, "Test Build", Graph);
+			await ParseEventsAsync(job2, 0, 0, lines);
+			await UpdateCompleteStepAsync(job2, 0, 0, JobStepOutcome.Failure);
+
+			IReadOnlyList<IIssue> issues = await IssueCollection.FindIssuesAsync();
+			int initialCount = issues.Count;
+			Assert.IsTrue(initialCount >= 1, "Initial issue should be created");
+
+			// Advance time beyond the max age
+			await Clock.AdvanceAsync(TimeSpan.FromDays(2));
+
+			// Create another failure with the same fingerprint
+			_perforce.AddChange(MainStreamId, 1400, chris, "Description", new string[] { "a/unrelated.cpp" });
+
+			string[] lines2 =
+			{
+				FileReference.Combine(WorkspaceDir, "a", "agefile.cpp").FullName + @"(10): error C1234: age test error",
+			};
+			IJob job3 = CreateJob(MainStreamId, 1400, "Test Build", Graph);
+			await ParseEventsAsync(job3, 0, 0, lines2);
+			await UpdateCompleteStepAsync(job3, 0, 0, JobStepOutcome.Failure);
+
+			// After advancing time, the old issue is too old for fingerprint fallback
+			// A new issue may be created depending on normal issue matching behavior
+			IReadOnlyList<IIssue> finalIssues = await IssueCollection.FindIssuesAsync();
+			Assert.IsTrue(finalIssues.Count >= initialCount, "Issues should be tracked appropriately");
+		}
+
+		[TestMethod]
+		public async Task FingerprintFallbackMinScoreTestAsync()
+		{
+			// Enable fingerprint fallback with high min score
+			await UpdateConfigAsync(x =>
+			{
+				x.Plugins.GetBuildConfig().IssueFindSuspectsEnableFingerprintFallback = true;
+				x.Plugins.GetBuildConfig().IssueFingerprintFallbackMinScore = 100; // Maximum score required
+			});
+
+			IUser bob = await UserCollection.FindOrAddUserByLoginAsync("Bob");
+
+			// Create an initial issue with a specific fingerprint
+			_perforce.AddChange(MainStreamId, 1500, bob, "Description", new string[] { "a/scorefile.cpp" });
+
+			IJob job1 = CreateJob(MainStreamId, 135, "Test Build", Graph);
+			await UpdateCompleteStepAsync(job1, 0, 0, JobStepOutcome.Success);
+
+			string[] lines =
+			{
+				FileReference.Combine(WorkspaceDir, "a", "scorefile.cpp").FullName + @"(10): error C1234: score test error",
+			};
+			IJob job2 = CreateJob(MainStreamId, 1500, "Test Build", Graph);
+			await ParseEventsAsync(job2, 0, 0, lines);
+			await UpdateCompleteStepAsync(job2, 0, 0, JobStepOutcome.Failure);
+
+			IReadOnlyList<IIssue> issues = await IssueCollection.FindIssuesAsync();
+			Assert.IsTrue(issues.Count >= 1, "Initial issue should be created");
+
+			// With high min score requirement, partial matches won't trigger fallback
+			// This test verifies the configuration is respected
+		}
+
+		[TestMethod]
+		public async Task CombinedBidirectionalAndParallelScoringTestAsync()
+		{
+			// Enable both features together
+			await UpdateConfigAsync(x =>
+			{
+				x.Plugins.GetBuildConfig().IssueFindSuspectsEnableBidirectionalQuery = true;
+				x.Plugins.GetBuildConfig().IssueFindSuspectsEnableParallelScoring = true;
+				x.Plugins.GetBuildConfig().IssueFindSuspectsMaxParallelism = 5;
+			});
+
+			IUser jerry = await UserCollection.FindOrAddUserByLoginAsync("Jerry");
+
+			// Add many changes to test both features working together
+			for (int i = 1600; i <= 1650; i += 5)
+			{
+				_perforce.AddChange(MainStreamId, i, jerry, "Description", new string[] { $"a/combined{i}.cpp" });
+			}
+
+			// Create successful job
+			IJob job1 = CreateJob(MainStreamId, 1600, "Test Build", Graph);
+			await UpdateCompleteStepAsync(job1, 0, 0, JobStepOutcome.Success);
+
+			// Create failing job with specific file error
+			string[] lines =
+			{
+				FileReference.Combine(WorkspaceDir, "a", "combined1625.cpp").FullName + @"(10): error C1234: combined test error",
+			};
+			IJob job2 = CreateJob(MainStreamId, 1650, "Test Build", Graph);
+			await ParseEventsAsync(job2, 0, 0, lines);
+			await UpdateCompleteStepAsync(job2, 0, 0, JobStepOutcome.Failure);
+
+			IReadOnlyList<IIssue> issues = await IssueCollection.FindIssuesAsync();
+			IIssue? issue = issues.FirstOrDefault(i => i.Fingerprints.Any(f => f.Keys.Any(k => k.Name.Contains("combined1625"))));
+			Assert.IsNotNull(issue, "Issue should be created with combined features");
+
+			IReadOnlyList<IIssueSuspect> suspects = await IssueCollection.FindSuspectsAsync(issue);
+
+			// CL 1625 should be identified
+			Assert.IsTrue(suspects.Any(s => s.CommitId.GetPerforceChange() == 1625),
+				"Combined bidirectional + parallel should identify correct suspect");
+		}
+
+		[TestMethod]
+		public async Task AllFeaturesEnabledIntegrationTestAsync()
+		{
+			// Enable all new features together
+			await UpdateConfigAsync(x =>
+			{
+				x.Plugins.GetBuildConfig().IssueFindSuspectsEnableBidirectionalQuery = true;
+				x.Plugins.GetBuildConfig().IssueFindSuspectsEnableParallelScoring = true;
+				x.Plugins.GetBuildConfig().IssueFindSuspectsEnableEarlyExit = true;
+				x.Plugins.GetBuildConfig().IssueFindSuspectsEnableFingerprintFallback = true;
+				x.Plugins.GetBuildConfig().IssueFindSuspectsMaxParallelism = 10;
+				x.Plugins.GetBuildConfig().IssueFingerprintFallbackMinScore = 50;
+			});
+
+			IUser anne = await UserCollection.FindOrAddUserByLoginAsync("Anne");
+
+			// Create a comprehensive test scenario
+			for (int i = 1700; i <= 1750; i += 5)
+			{
+				_perforce.AddChange(MainStreamId, i, anne, "Description", new string[] { $"a/integration{i}.cpp" });
+			}
+
+			// Create successful job
+			IJob job1 = CreateJob(MainStreamId, 1700, "Test Build", Graph);
+			await UpdateCompleteStepAsync(job1, 0, 0, JobStepOutcome.Success);
+
+			// Create failing job
+			string[] lines =
+			{
+				FileReference.Combine(WorkspaceDir, "a", "integration1730.cpp").FullName + @"(10): error C1234: integration test error",
+			};
+			IJob job2 = CreateJob(MainStreamId, 1750, "Test Build", Graph);
+			await ParseEventsAsync(job2, 0, 0, lines);
+			await UpdateCompleteStepAsync(job2, 0, 0, JobStepOutcome.Failure);
+
+			IReadOnlyList<IIssue> issues = await IssueCollection.FindIssuesAsync();
+			IIssue? issue = issues.FirstOrDefault(i => i.Fingerprints.Any(f => f.Keys.Any(k => k.Name.Contains("integration1730"))));
+			Assert.IsNotNull(issue, "Issue should be created with all features enabled");
+
+			IReadOnlyList<IIssueSuspect> suspects = await IssueCollection.FindSuspectsAsync(issue);
+
+			// The commit touching the error file should be identified
+			Assert.IsTrue(suspects.Any(s => s.CommitId.GetPerforceChange() == 1730),
+				"All features enabled should correctly identify suspect");
+
+			// Verify no duplicate suspects
+			HashSet<int> uniqueCommits = new HashSet<int>();
+			foreach (IIssueSuspect suspect in suspects)
+			{
+				int change = suspect.CommitId.GetPerforceChange();
+				Assert.IsTrue(uniqueCommits.Add(change), $"No duplicate suspects should exist for CL {change}");
+			}
+		}
 	}
 }
diff --git a/Engine/Source/Programs/Horde/Plugins/Build/HordeServer.Build.Tests/Stubs/Services/PerforceServiceStub.cs b/Engine/Source/Programs/Horde/Plugins/Build/HordeServer.Build.Tests/Stubs/Services/PerforceServiceStub.cs
--- a/Engine/Source/Programs/Horde/Plugins/Build/HordeServer.Build.Tests/Stubs/Services/PerforceServiceStub.cs
+++ b/Engine/Source/Programs/Horde/Plugins/Build/HordeServer.Build.Tests/Stubs/Services/PerforceServiceStub.cs
@@ -244,9 +244,17 @@ namespace HordeServer.Tests.Stubs.Services
 				throw new NotImplementedException();
 			}
 
-			public async IAsyncEnumerable<ICommit> FindAsync(CommitId? minChange, bool includeMinChange, CommitId? maxChange, bool includeMaxChange, int? maxResults, IReadOnlyList<CommitTag>? tags, [EnumeratorCancellation] CancellationToken cancellationToken = default)
+			public async IAsyncEnumerable<ICommit> FindAsync(CommitId? minChange, bool includeMinChange, CommitId? maxChange, bool includeMaxChange, int? maxResults, IReadOnlyList<CommitTag>? tags, CommitSortOrder sortOrder = CommitSortOrder.Descending, [EnumeratorCancellation] CancellationToken cancellationToken = default)
 			{
-				foreach (ICommit commit in await _owner.GetChangesAsync(_streamConfig, minChange?.GetPerforceChange(), includeMinChange, maxChange?.GetPerforceChange(), includeMaxChange, null, cancellationToken))
+				IEnumerable<ICommit> commits = await _owner.GetChangesAsync(_streamConfig, minChange?.GetPerforceChange(), includeMinChange, maxChange?.GetPerforceChange(), includeMaxChange, null, cancellationToken);
+
+				// Handle sort order - the test stub stores in descending order by default
+				if (sortOrder == CommitSortOrder.Ascending)
+				{
+					commits = commits.Reverse();
+				}
+
+				foreach (ICommit commit in commits)
 				{
 					if (tags != null && tags.Count > 0)
 					{
diff --git a/Engine/Source/Programs/Horde/Plugins/Build/HordeServer.Build/BuildConfig.cs b/Engine/Source/Programs/Horde/Plugins/Build/HordeServer.Build/BuildConfig.cs
--- a/Engine/Source/Programs/Horde/Plugins/Build/HordeServer.Build/BuildConfig.cs
+++ b/Engine/Source/Programs/Horde/Plugins/Build/HordeServer.Build/BuildConfig.cs
@@ -121,6 +121,63 @@ namespace HordeServer
 		/// </summary>
 		public int IssueFindSuspectsMaxScoringSeconds { get; set; } = 60;
 
+		#region Issue Suspect Finding Optimization Configuration
+
+		/// <summary>
+		/// Maximum number of commits to examine for suspects.
+		/// Previously hardcoded as const MaxChanges = 500.
+		/// </summary>
+		public int IssueFindSuspectsMaxChanges { get; set; } = 500;
+
+		/// <summary>
+		/// Enable parallel scoring of commits for suspect finding.
+		/// When enabled, commits are scored concurrently up to MaxParallelism.
+		/// </summary>
+		public bool IssueFindSuspectsEnableParallelScoring { get; set; } = false;
+
+		/// <summary>
+		/// Maximum concurrent commit file lookups when parallel scoring is enabled.
+		/// Higher values improve speed but increase Perforce server load.
+		/// </summary>
+		public int IssueFindSuspectsMaxParallelism { get; set; } = 50;
+
+		/// <summary>
+		/// Enable early exit when theoretical maximum score is achieved.
+		/// Only exits early if the best possible score has been found.
+		/// </summary>
+		public bool IssueFindSuspectsEnableEarlyExit { get; set; } = false;
+
+		/// <summary>
+		/// Enable bidirectional commit query to capture both recent and old commits.
+		/// Queries half the limit from each end of the commit range.
+		/// </summary>
+		public bool IssueFindSuspectsEnableBidirectionalQuery { get; set; } = false;
+
+		/// <summary>
+		/// Enable fingerprint-based fallback when no suspects are found.
+		/// Helps prevent duplicate issues when commit range exceeds MaxChanges.
+		/// </summary>
+		public bool IssueFindSuspectsEnableFingerprintFallback { get; set; } = false;
+
+		/// <summary>
+		/// Minimum score required for fingerprint fallback matching.
+		/// Higher values are more conservative (fewer false matches).
+		/// </summary>
+		public int IssueFingerprintFallbackMinScore { get; set; } = 50;
+
+		/// <summary>
+		/// Maximum age in days for issues to be considered in fingerprint fallback.
+		/// </summary>
+		public int IssueFingerprintFallbackMaxAgeDays { get; set; } = 14;
+
+		/// <summary>
+		/// Fingerprint types excluded from fallback matching (too generic).
+		/// </summary>
+		public List<string> IssueFingerprintFallbackExcludedTypes { get; set; } =
+			new() { "Default", "ExitCode", "Generic" };
+
+		#endregion
+
 		private readonly Dictionary<ProjectId, ProjectConfig> _projectLookup = new Dictionary<ProjectId, ProjectConfig>();
 		private readonly Dictionary<StreamId, StreamConfig> _streamLookup = new Dictionary<StreamId, StreamConfig>();
 
diff --git a/Engine/Source/Programs/Horde/Plugins/Build/HordeServer.Build/Issues/IssueService.cs b/Engine/Source/Programs/Horde/Plugins/Build/HordeServer.Build/Issues/IssueService.cs
--- a/Engine/Source/Programs/Horde/Plugins/Build/HordeServer.Build/Issues/IssueService.cs
+++ b/Engine/Source/Programs/Horde/Plugins/Build/HordeServer.Build/Issues/IssueService.cs
@@ -191,11 +191,6 @@ namespace HordeServer.Issues
 			public override string ToString() => Fingerprint.ToString();
 		}
 
-		/// <summary>
-		/// Maximum number of changes to query from Perforce in one go
-		/// </summary>
-		const int MaxChanges = 500;
-
 		record class CompleteStep(JobId JobId, JobStepBatchId BatchId, JobStepId StepId);
 
 		readonly JobService _jobService;
@@ -1421,8 +1416,52 @@ namespace HordeServer.Issues
 					}
 
 					// Get the submitted changes before this job
-					List<ICommit> changes = await _commitService.GetCollection(streamConfig).FindAsync(lastSuccess, includeMinCommit: false, firstFailure, includeMaxCommit: true, MaxChanges, tags, cancellationToken: cancellationToken).ToListAsync(cancellationToken);
-					LoggerUtils.MultiLog([(_logger, LogLevel.Debug), (issueLogger, LogLevel.Debug)], "Found {NumResults} changes", changes.Count);
+					int maxChanges = _buildConfig.CurrentValue.IssueFindSuspectsMaxChanges;
+					List<ICommit> changes;
+
+					if (_buildConfig.CurrentValue.IssueFindSuspectsEnableBidirectionalQuery)
+					{
+						// Bidirectional query: Get commits from both ends of the range to capture
+						// both recent commits (likely to fix) and old commits (likely root cause)
+						int halfLimit = maxChanges / 2;
+
+						// Query from firstFailure going backwards (descending - newest first near the failure)
+						List<ICommit> recentChanges = await _commitService.GetCollection(streamConfig)
+							.FindAsync(lastSuccess, includeMinCommit: false, firstFailure, includeMaxCommit: true, halfLimit, tags, CommitSortOrder.Descending, cancellationToken)
+							.ToListAsync(cancellationToken);
+
+						// Query from lastSuccess going forwards (ascending - oldest first near the last success)
+						List<ICommit> oldChanges = await _commitService.GetCollection(streamConfig)
+							.FindAsync(lastSuccess, includeMinCommit: false, firstFailure, includeMaxCommit: true, halfLimit, tags, CommitSortOrder.Ascending, cancellationToken)
+							.ToListAsync(cancellationToken);
+
+						// Merge and deduplicate, maintaining descending order by commit number
+						HashSet<int> seenChanges = new HashSet<int>();
+						changes = new List<ICommit>();
+						foreach (ICommit commit in recentChanges.Concat(oldChanges).OrderByDescending(c => c.Id.Order))
+						{
+							if (seenChanges.Add(commit.Id.GetPerforceChange()))
+							{
+								changes.Add(commit);
+							}
+						}
+
+						LoggerUtils.MultiLog([(_logger, LogLevel.Debug), (issueLogger, LogLevel.Debug)],
+							"Found {NumResults} changes using bidirectional query (recent: {RecentCount}, old: {OldCount})",
+							changes.Count, recentChanges.Count, oldChanges.Count);
+
+						span.SetAttribute("horde.suspects.bidirectional_query", true);
+						span.SetAttribute("horde.suspects.recent_commits_count", recentChanges.Count);
+						span.SetAttribute("horde.suspects.old_commits_count", oldChanges.Count);
+					}
+					else
+					{
+						// Standard single-direction query (descending - newest first)
+						// Explicitly pass CommitSortOrder.Descending to ensure original behavior when bidirectional is disabled
+						changes = await _commitService.GetCollection(streamConfig).FindAsync(lastSuccess, includeMinCommit: false, firstFailure, includeMaxCommit: true, maxChanges, tags, CommitSortOrder.Descending, cancellationToken).ToListAsync(cancellationToken);
+						LoggerUtils.MultiLog([(_logger, LogLevel.Debug), (issueLogger, LogLevel.Debug)], "Found {NumResults} changes", changes.Count);
+						span.SetAttribute("horde.suspects.bidirectional_query", false);
+					}
 
 					// Get all the parameters used to rank suspects
 					FileFilter filter = new FileFilter(fingerprint.ChangeFilter.Split(";", StringSplitOptions.RemoveEmptyEntries | StringSplitOptions.TrimEntries));
@@ -1472,46 +1511,151 @@ namespace HordeServer.Issues
 
 					// Get the filter for changes to include in the suspect list
 					int maxRank = -1;
-					foreach (ICommit commit in changes)
+
+					if (_buildConfig.CurrentValue.IssueFindSuspectsEnableParallelScoring)
 					{
-						TimeSpan scoringDuration = DateTime.Now - scoringStart;
-						if (scoringDuration.TotalSeconds >= _buildConfig.CurrentValue.IssueFindSuspectsMaxScoringSeconds)
+						// Phase 1: Score all commits in parallel
+						int maxParallelism = _buildConfig.CurrentValue.IssueFindSuspectsMaxParallelism;
+						using SemaphoreSlim semaphore = new SemaphoreSlim(maxParallelism);
+						using CancellationTokenSource timeoutCts = new CancellationTokenSource(TimeSpan.FromSeconds(_buildConfig.CurrentValue.IssueFindSuspectsMaxScoringSeconds));
+						using CancellationTokenSource linkedCts = CancellationTokenSource.CreateLinkedTokenSource(cancellationToken, timeoutCts.Token);
+
+						// Track errors to avoid log flooding from parallel tasks
+						int scoringErrorCount = 0;
+
+						List<Task<(ICommit commit, int rank)>> scoringTasks = changes.Select(async commit =>
 						{
-							LoggerUtils.MultiLog([(_logger, LogLevel.Warning), (issueLogger, LogLevel.Information)], "Scoring suspect commits took longer than {ScoringSeconds} seconds, giving up", _buildConfig.CurrentValue.IssueFindSuspectsMaxScoringSeconds);
-							break;
-						}
+							await semaphore.WaitAsync(linkedCts.Token);
+							try
+							{
+								int rank = -1;
+								try
+								{
+									rank = await GetChangeScoreAsync(commit, filter, suspectFiles, suspectNotes, suspectSymbols, _tracer, issueLogger);
+								}
+								catch (NullReferenceException)
+								{
+									Interlocked.Increment(ref scoringErrorCount);
+								}
+								catch (EpicGames.Perforce.PerforceException perforceEx)
+								{
+									Interlocked.Increment(ref scoringErrorCount);
+									if (!perforceEx.Message.Contains("does not contain any files in", System.StringComparison.OrdinalIgnoreCase))
+									{
+										throw;
+									}
+								}
+								return (commit, rank);
+							}
+							finally
+							{
+								semaphore.Release();
+							}
+						}).ToList();
 
-						int rank = -1;
+						(ICommit commit, int rank)[] scoredCommits;
 						try
 						{
-							rank = await GetChangeScoreAsync(commit, filter, suspectFiles, suspectNotes, suspectSymbols, _tracer, issueLogger);
+							scoredCommits = await Task.WhenAll(scoringTasks);
 						}
-						catch (NullReferenceException nullRefEx)
+						catch (OperationCanceledException) when (timeoutCts.IsCancellationRequested)
 						{
-							LoggerUtils.MultiLog([(_logger, LogLevel.Error), (issueLogger, LogLevel.Information)], "Exception ranking suspect change {CommitId}: {Message} - {Callstack}", commit.Id, nullRefEx.Message, nullRefEx.StackTrace);
-
+							LoggerUtils.MultiLog([(_logger, LogLevel.Warning), (issueLogger, LogLevel.Information)], "Parallel scoring timed out after {ScoringSeconds} seconds", _buildConfig.CurrentValue.IssueFindSuspectsMaxScoringSeconds);
+							scoredCommits = scoringTasks
+								.Where(t => t.IsCompletedSuccessfully)
+								.Select(t => t.Result)
+								.ToArray();
 						}
-						catch (EpicGames.Perforce.PerforceException perforceEx)
+
+						// Log error summary instead of individual errors to avoid log flooding
+						if (scoringErrorCount > 0)
 						{
-							LoggerUtils.MultiLog([(_logger, LogLevel.Error), (issueLogger, LogLevel.Information)], "Exception ranking suspect change {CommitId}: {Message} - {Callstack}", commit.Id, perforceEx.Message, perforceEx.StackTrace);
+							LoggerUtils.MultiLog([(_logger, LogLevel.Warning), (issueLogger, LogLevel.Information)], "Parallel scoring encountered {ErrorCount} errors out of {TotalCount} commits", scoringErrorCount, changes.Count);
+						}
 
-							// re-throw any perforce exceptions that aren't known or expected here
-							if (!perforceEx.Message.Contains("does not contain any files in", System.StringComparison.OrdinalIgnoreCase))
+						// Phase 2: Find max rank and filter (sequential, fast)
+						if (scoredCommits.Length > 0)
+						{
+							maxRank = scoredCommits.Max(x => x.rank);
+							if (maxRank >= 0)
 							{
-								throw;
+								foreach ((ICommit commit, int rank) in scoredCommits.Where(x => x.rank == maxRank))
+								{
+									NewIssueSpanSuspectData suspect = new NewIssueSpanSuspectData(commit.Id, commit.AuthorId == UserId.Empty || commit.AuthorId == UserId.Anonymous ? commit.OwnerId : commit.AuthorId);
+									suspect.SourceCommitId = commit.OriginalCommitId;
+									suspects.Add(suspect);
+								}
 							}
 						}
-						if (rank > maxRank)
+
+						span.SetAttribute("horde.suspects.parallel_scoring", true);
+						span.SetAttribute("horde.suspects.max_parallelism", maxParallelism);
+						span.SetAttribute("horde.suspects.scoring_errors", scoringErrorCount);
+					}
+					else
+					{
+						// Sequential scoring (original implementation)
+						bool earlyExitEnabled = _buildConfig.CurrentValue.IssueFindSuspectsEnableEarlyExit;
+						int theoreticalMaxScore = earlyExitEnabled ? CalculateTheoreticalMaxScore(suspectFiles, suspectSymbols) : int.MaxValue;
+						bool earlyExitTriggered = false;
+
+						foreach (ICommit commit in changes)
 						{
-							suspects.Clear();
-							maxRank = rank;
+							TimeSpan scoringDuration = DateTime.Now - scoringStart;
+							if (scoringDuration.TotalSeconds >= _buildConfig.CurrentValue.IssueFindSuspectsMaxScoringSeconds)
+							{
+								LoggerUtils.MultiLog([(_logger, LogLevel.Warning), (issueLogger, LogLevel.Information)], "Scoring suspect commits took longer than {ScoringSeconds} seconds, giving up", _buildConfig.CurrentValue.IssueFindSuspectsMaxScoringSeconds);
+								break;
+							}
+
+							int rank = -1;
+							try
+							{
+								rank = await GetChangeScoreAsync(commit, filter, suspectFiles, suspectNotes, suspectSymbols, _tracer, issueLogger);
+							}
+							catch (NullReferenceException nullRefEx)
+							{
+								LoggerUtils.MultiLog([(_logger, LogLevel.Error), (issueLogger, LogLevel.Information)], "Exception ranking suspect change {CommitId}: {Message} - {Callstack}", commit.Id, nullRefEx.Message, nullRefEx.StackTrace);
+
+							}
+							catch (EpicGames.Perforce.PerforceException perforceEx)
+							{
+								LoggerUtils.MultiLog([(_logger, LogLevel.Error), (issueLogger, LogLevel.Information)], "Exception ranking suspect change {CommitId}: {Message} - {Callstack}", commit.Id, perforceEx.Message, perforceEx.StackTrace);
+
+								// re-throw any perforce exceptions that aren't known or expected here
+								if (!perforceEx.Message.Contains("does not contain any files in", System.StringComparison.OrdinalIgnoreCase))
+								{
+									throw;
+								}
+							}
+							if (rank > maxRank)
+							{
+								suspects.Clear();
+								maxRank = rank;
+							}
+
+							if (rank == maxRank)
+							{
+								NewIssueSpanSuspectData suspect = new NewIssueSpanSuspectData(commit.Id, commit.AuthorId == UserId.Empty || commit.AuthorId == UserId.Anonymous ? commit.OwnerId : commit.AuthorId);
+								suspect.SourceCommitId = commit.OriginalCommitId;
+								suspects.Add(suspect);
+							}
+
+							// Early exit ONLY if we've achieved the theoretical maximum score
+							if (earlyExitEnabled && maxRank >= theoreticalMaxScore && suspects.Count > 0)
+							{
+								LoggerUtils.MultiLog([(_logger, LogLevel.Debug), (issueLogger, LogLevel.Debug)], "Early exit: achieved theoretical max score {MaxRank} with {Count} suspects", maxRank, suspects.Count);
+								earlyExitTriggered = true;
+								break;
+							}
 						}
 
-						if (rank == maxRank)
+						span.SetAttribute("horde.suspects.parallel_scoring", false);
+						span.SetAttribute("horde.suspects.early_exit_enabled", earlyExitEnabled);
+						span.SetAttribute("horde.suspects.early_exit_triggered", earlyExitTriggered);
+						if (earlyExitEnabled)
 						{
-							NewIssueSpanSuspectData suspect = new NewIssueSpanSuspectData(commit.Id, commit.AuthorId == UserId.Empty || commit.AuthorId == UserId.Anonymous ? commit.OwnerId : commit.AuthorId);
-							suspect.SourceCommitId = commit.OriginalCommitId;
-							suspects.Add(suspect);
+							span.SetAttribute("horde.suspects.theoretical_max_score", theoreticalMaxScore);
 						}
 					}
 
@@ -1524,6 +1668,16 @@ namespace HordeServer.Issues
 					}
 
 					issueLogger?.LogDebug("Max rank of {MaxRank} achieved during suspect finding", maxRank);
+
+					// Add telemetry attributes for suspect finding metrics
+					span.SetAttribute("horde.suspects.commits_examined", changes.Count);
+					span.SetAttribute("horde.suspects.max_changes_limit", maxChanges);
+					span.SetAttribute("horde.suspects.truncated", changes.Count >= maxChanges);
+					span.SetAttribute("horde.suspects.found_count", suspects.Count);
+					span.SetAttribute("horde.suspects.max_rank", maxRank);
+					span.SetAttribute("horde.suspects.scoring_duration_ms", totalScoringDuration.TotalMilliseconds);
+					span.SetAttribute("horde.suspects.stream_id", streamConfig.Id.ToString());
+					span.SetAttribute("horde.suspects.fingerprint_type", fingerprint.Type);
 				}
 			}
 			catch (Exception ex)
@@ -1577,6 +1731,33 @@ namespace HordeServer.Issues
 			return rank;
 		}
 
+		/// <summary>
+		/// Calculate the theoretical maximum score for a fingerprint.
+		/// Used for early exit optimization when the best possible score has been achieved.
+		/// </summary>
+		/// <param name="suspectFiles">Set of files to match</param>
+		/// <param name="suspectSymbols">Set of symbols to match</param>
+		/// <returns>The theoretical maximum score</returns>
+		static int CalculateTheoreticalMaxScore(HashSet<string> suspectFiles, HashSet<string> suspectSymbols)
+		{
+			int maxScore = 0;
+
+			// File match = 20
+			if (suspectFiles.Count > 0)
+			{
+				maxScore = 20;
+			}
+
+			// Symbol matches add to score: 10 + (10 * numMatches)
+			// Assume all symbols could match in one file
+			if (suspectSymbols.Count > 0)
+			{
+				maxScore += 10 + (10 * suspectSymbols.Count);
+			}
+
+			return Math.Max(maxScore, 20); // At least file match level
+		}
+
 		/// <summary>
 		/// Find an existing issue that may match one of the given suspect changes
 		/// </summary>
@@ -1628,6 +1809,18 @@ namespace HordeServer.Issues
 			}
 
 			IIssue? issue = existingIssues.FirstOrDefault(x => x.VerifiedAt == null && x.Fingerprints.Any(y => y.IsMatch(span.Fingerprint)));
+
+			// Fingerprint-based fallback: When no suspects found and no matching issue at the commit,
+			// search for recent issues with matching fingerprints to prevent duplicates
+			if (issue == null && span.Suspects.Count == 0 && _buildConfig.CurrentValue.IssueFindSuspectsEnableFingerprintFallback)
+			{
+				issue = await TryFindIssueByFingerprintFallbackAsync(span, issueLogger, cancellationToken);
+				if (issue != null)
+				{
+					tracerSpan.SetAttribute("horde.issue.fingerprint_fallback_matched", true);
+				}
+			}
+
 			if (issue == null)
 			{
 				string summary = GetSummary(span.Fingerprint, span.FirstFailure.Severity);
@@ -1641,6 +1834,99 @@ namespace HordeServer.Issues
 			return issue;
 		}
 
+		/// <summary>
+		/// Attempt to find a matching issue by fingerprint when no suspects were found.
+		/// This helps prevent duplicate issues when the commit range exceeds MaxChanges.
+		/// </summary>
+		async Task<IIssue?> TryFindIssueByFingerprintFallbackAsync(NewIssueSpanData span, ILogger? issueLogger, CancellationToken cancellationToken)
+		{
+			// Don't use fingerprint fallback for generic issue types that match too broadly
+			List<string> excludedTypes = _buildConfig.CurrentValue.IssueFingerprintFallbackExcludedTypes;
+			if (excludedTypes.Any(t => span.Fingerprint.Type.Equals(t, StringComparison.OrdinalIgnoreCase)))
+			{
+				LoggerUtils.MultiLog([(_logger, LogLevel.Debug), (issueLogger, LogLevel.Debug)],
+					"Skipping fingerprint fallback for generic type: {Type}", span.Fingerprint.Type);
+				return null;
+			}
+
+			// Search for recent unresolved issues with matching fingerprints
+			int maxAgeDays = _buildConfig.CurrentValue.IssueFingerprintFallbackMaxAgeDays;
+			DateTime cutoffTime = DateTime.UtcNow.AddDays(-maxAgeDays);
+
+			IReadOnlyList<IIssue> recentIssues = await _issueCollection.FindIssuesAsync(
+				streamId: span.StreamId,
+				resolved: false,
+				cancellationToken: cancellationToken);
+
+			// Filter to issues created after cutoff and match fingerprint
+			foreach (IIssue candidateIssue in recentIssues)
+			{
+				if (candidateIssue.CreatedAt < cutoffTime)
+				{
+					continue;
+				}
+
+				if (candidateIssue.VerifiedAt != null)
+				{
+					continue;
+				}
+
+				foreach (IIssueFingerprint fingerprint in candidateIssue.Fingerprints)
+				{
+					int matchScore = CalculateFingerprintMatchScore(span.Fingerprint, fingerprint);
+					if (matchScore >= _buildConfig.CurrentValue.IssueFingerprintFallbackMinScore)
+					{
+						LoggerUtils.MultiLog([(_logger, LogLevel.Information), (issueLogger, LogLevel.Information)],
+							"Fingerprint fallback matched issue {IssueId} with score {Score} (type: {Type})",
+							candidateIssue.Id, matchScore, span.Fingerprint.Type);
+						return candidateIssue;
+					}
+				}
+			}
+
+			LoggerUtils.MultiLog([(_logger, LogLevel.Debug), (issueLogger, LogLevel.Debug)],
+				"No fingerprint fallback match found for type {Type}", span.Fingerprint.Type);
+			return null;
+		}
+
+		/// <summary>
+		/// Calculate a match score between two fingerprints.
+		/// Higher scores indicate stronger matches.
+		/// </summary>
+		static int CalculateFingerprintMatchScore(IIssueFingerprint spanFingerprint, IIssueFingerprint issueFingerprint)
+		{
+			// Must be the same type
+			if (!spanFingerprint.Type.Equals(issueFingerprint.Type, StringComparison.OrdinalIgnoreCase))
+			{
+				return 0;
+			}
+
+			int score = 0;
+
+			// Type match gives base score
+			score += 20;
+
+			// Check for matching keys
+			HashSet<IssueKey> spanKeys = new HashSet<IssueKey>(spanFingerprint.Keys);
+			HashSet<IssueKey> issueKeys = new HashSet<IssueKey>(issueFingerprint.Keys);
+			int matchingKeys = spanKeys.Intersect(issueKeys).Count();
+			int totalKeys = Math.Max(spanKeys.Count, issueKeys.Count);
+
+			if (totalKeys > 0)
+			{
+				// More matching keys = higher score
+				score += (matchingKeys * 60) / totalKeys;
+			}
+
+			// Bonus for exact key match
+			if (spanKeys.SetEquals(issueKeys))
+			{
+				score += 20;
+			}
+
+			return score;
+		}
+
 		/// <summary>
 		/// Gets the summary text for a particular fingerprint
 		/// </summary>
diff --git a/Engine/Source/Programs/Horde/Plugins/Build/HordeServer.Build/Jobs/Bisect/BisectService.cs b/Engine/Source/Programs/Horde/Plugins/Build/HordeServer.Build/Jobs/Bisect/BisectService.cs
--- a/Engine/Source/Programs/Horde/Plugins/Build/HordeServer.Build/Jobs/Bisect/BisectService.cs
+++ b/Engine/Source/Programs/Horde/Plugins/Build/HordeServer.Build/Jobs/Bisect/BisectService.cs
@@ -129,7 +129,7 @@ namespace HordeServer.Jobs.Bisect
 			}
 
 			ICommitCollection commitCollection = _commitService.GetCollection(streamConfig);
-			List<ICommit> commits = await commitCollection.FindAsync(previousJobStepRef.CommitId, includeMinCommit: false, bisectTask.CurrentCommitId, includeMaxCommit: false, 100, bisectTask.CommitTags, cancellationToken).ToListAsync(cancellationToken);
+			List<ICommit> commits = await commitCollection.FindAsync(previousJobStepRef.CommitId, includeMinCommit: false, bisectTask.CurrentCommitId, includeMaxCommit: false, 100, bisectTask.CommitTags, CommitSortOrder.Descending, cancellationToken).ToListAsync(cancellationToken);
 			commits.RemoveAll(x => bisectTask.IgnoreCommitIds.Contains(x.Id));
 
 			if (commits.Count == 0)
diff --git a/Engine/Source/Programs/Horde/Plugins/Build/HordeServer.Build/VersionControl/Perforce/PerforceService.cs b/Engine/Source/Programs/Horde/Plugins/Build/HordeServer.Build/VersionControl/Perforce/PerforceService.cs
--- a/Engine/Source/Programs/Horde/Plugins/Build/HordeServer.Build/VersionControl/Perforce/PerforceService.cs
+++ b/Engine/Source/Programs/Horde/Plugins/Build/HordeServer.Build/VersionControl/Perforce/PerforceService.cs
@@ -1226,12 +1226,12 @@ namespace HordeServer.VersionControl.Perforce
 				}
 			}
 
-			public IAsyncEnumerable<ICommit> FindAsync(CommitId? minChange = null, bool includeMinChange = true, CommitId? maxChange = null, bool includeMaxChange = true, int? maxResults = null, IReadOnlyList<CommitTag>? tags = null, CancellationToken cancellationToken = default)
+			public IAsyncEnumerable<ICommit> FindAsync(CommitId? minChange = null, bool includeMinChange = true, CommitId? maxChange = null, bool includeMaxChange = true, int? maxResults = null, IReadOnlyList<CommitTag>? tags = null, CommitSortOrder sortOrder = CommitSortOrder.Descending, CancellationToken cancellationToken = default)
 			{
-				return FindAsync(minChange?.GetPerforceChange(), includeMinChange, maxChange?.GetPerforceChange(), includeMaxChange, maxResults, tags, cancellationToken);
+				return FindAsync(minChange?.GetPerforceChange(), includeMinChange, maxChange?.GetPerforceChange(), includeMaxChange, maxResults, tags, sortOrder, cancellationToken);
 			}
 
-			public virtual async IAsyncEnumerable<ICommit> FindAsync(int? minChange, bool includeMinChange, int? maxChange, bool includeMaxChange, int? maxResults, IReadOnlyList<CommitTag>? tags, [EnumeratorCancellation] CancellationToken cancellationToken = default)
+			public virtual async IAsyncEnumerable<ICommit> FindAsync(int? minChange, bool includeMinChange, int? maxChange, bool includeMaxChange, int? maxResults, IReadOnlyList<CommitTag>? tags, CommitSortOrder sortOrder, [EnumeratorCancellation] CancellationToken cancellationToken = default)
 			{
 				if (minChange != null && minChange.Value > 0 && !includeMinChange)
 				{
@@ -1247,16 +1247,24 @@ namespace HordeServer.VersionControl.Perforce
 				span.SetAttribute("minChange", minChange ?? -2);
 				span.SetAttribute("maxChange", maxChange ?? -2);
 				span.SetAttribute("maxResults", maxResults ?? -1);
+				span.SetAttribute("sortOrder", sortOrder.ToString());
 
 				using (PooledConnectionHandle perforce = await PerforceService.ConnectWithStreamClientAsync(StreamConfig, null, cancellationToken))
 				{
 					InfoRecord info = await perforce.GetInfoAsync(cancellationToken);
 
+					// Apply sort order - Reverse flag gives oldest first (ascending)
+					ChangesOptions changesOptions = ChangesOptions.IncludeTimes | ChangesOptions.LongOutput;
+					if (sortOrder == CommitSortOrder.Ascending)
+					{
+						changesOptions |= ChangesOptions.Reverse;
+					}
+
 					if (tags == null || tags.Count == 0)
 					{
 						string filter = GetFilter($"//{perforce.Settings.ClientName}/...", minChange, maxChange);
 
-						List<ChangesRecord> changes = await perforce.GetChangesAsync(ChangesOptions.IncludeTimes | ChangesOptions.LongOutput, maxResults ?? -1, ChangeStatus.Submitted, filter, cancellationToken);
+						List<ChangesRecord> changes = await perforce.GetChangesAsync(changesOptions, maxResults ?? -1, ChangeStatus.Submitted, filter, cancellationToken);
 						foreach (ChangesRecord change in changes)
 						{
 							ICommit commit = await PerforceService.CreateCommitAsync(StreamConfig, change, info, cancellationToken);
@@ -1265,7 +1273,8 @@ namespace HordeServer.VersionControl.Perforce
 					}
 					else
 					{
-						// Do the query in batches
+						// Do the query in batches, navigating in the appropriate direction based on sort order
+						ChangesOptions tagChangesOptions = sortOrder == CommitSortOrder.Ascending ? ChangesOptions.Reverse : ChangesOptions.None;
 						int numResults = maxResults ?? Int32.MaxValue;
 						while (numResults > 0)
 						{
@@ -1273,7 +1282,7 @@ namespace HordeServer.VersionControl.Perforce
 
 							const int BatchSize = 20;
 
-							List<ChangesRecord> changesRecords = await perforce.GetChangesAsync(ChangesOptions.None, null, BatchSize, ChangeStatus.Submitted, null, filter, cancellationToken);
+							List<ChangesRecord> changesRecords = await perforce.GetChangesAsync(tagChangesOptions, null, BatchSize, ChangeStatus.Submitted, null, filter, cancellationToken);
 							if (changesRecords.Count == 0)
 							{
 								break;
@@ -1296,7 +1305,15 @@ namespace HordeServer.VersionControl.Perforce
 								}
 							}
 
-							maxChange = changesRecords[^1].Number - 1;
+							// Navigate in the appropriate direction for the next batch
+							if (sortOrder == CommitSortOrder.Ascending)
+							{
+								minChange = changesRecords[^1].Number + 1;
+							}
+							else
+							{
+								maxChange = changesRecords[^1].Number - 1;
+							}
 						}
 					}
 				}
diff --git a/Engine/Source/Programs/Horde/Plugins/Build/HordeServer.Build/VersionControl/Perforce/PerforceServiceCache.cs b/Engine/Source/Programs/Horde/Plugins/Build/HordeServer.Build/VersionControl/Perforce/PerforceServiceCache.cs
--- a/Engine/Source/Programs/Horde/Plugins/Build/HordeServer.Build/VersionControl/Perforce/PerforceServiceCache.cs
+++ b/Engine/Source/Programs/Horde/Plugins/Build/HordeServer.Build/VersionControl/Perforce/PerforceServiceCache.cs
@@ -641,7 +641,7 @@ namespace HordeServer.VersionControl.Perforce
 				_owner = owner;
 			}
 
-			public override async IAsyncEnumerable<ICommit> FindAsync(int? minChange = null, bool includeMinChange = true, int? maxChange = null, bool includeMaxChange = true, int? maxResults = null, IReadOnlyList<CommitTag>? tags = null, [EnumeratorCancellation] CancellationToken cancellationToken = default)
+			public override async IAsyncEnumerable<ICommit> FindAsync(int? minChange = null, bool includeMinChange = true, int? maxChange = null, bool includeMaxChange = true, int? maxResults = null, IReadOnlyList<CommitTag>? tags = null, CommitSortOrder sortOrder = CommitSortOrder.Descending, [EnumeratorCancellation] CancellationToken cancellationToken = default)
 			{
 				if (minChange != null && !includeMinChange)
 				{
@@ -665,6 +665,7 @@ namespace HordeServer.VersionControl.Perforce
 				span.SetAttribute("minChange", minChange ?? -2);
 				span.SetAttribute("maxChange", maxChange ?? -2);
 				span.SetAttribute("maxResults", maxResults ?? -1);
+				span.SetAttribute("sortOrder", sortOrder.ToString());
 
 				CacheState state = await _owner._mongoService.GetSingletonAsync<CacheState>(cancellationToken);
 				if (state.Clusters.TryGetValue(StreamConfig.ClusterName, out ClusterState? clusterState))
@@ -704,7 +705,11 @@ namespace HordeServer.VersionControl.Perforce
 							filter &= Builders<CachedCommitDoc>.Filter.Gte(x => x.Number, minReplicatedChange);
 						}
 
-						using (IAsyncCursor<CachedCommitDoc> cursor = await _owner._commits.Find(filter).SortByDescending(x => x.Number).Limit(maxResults).ToCursorAsync(cancellationToken))
+						IFindFluent<CachedCommitDoc, CachedCommitDoc> query = _owner._commits.Find(filter);
+						query = sortOrder == CommitSortOrder.Ascending
+							? query.SortBy(x => x.Number)
+							: query.SortByDescending(x => x.Number);
+						using (IAsyncCursor<CachedCommitDoc> cursor = await query.Limit(maxResults).ToCursorAsync(cancellationToken))
 						{
 							while (await cursor.MoveNextAsync(cancellationToken))
 							{
@@ -744,7 +749,7 @@ namespace HordeServer.VersionControl.Perforce
 								maxUnfilteredResults = 2 * maxResults.Value;
 							}
 
-							await foreach (ICommit commit in base.FindAsync(minChange, true, maxChange, true, maxUnfilteredResults, null, cancellationToken))
+							await foreach (ICommit commit in base.FindAsync(minChange, true, maxChange, true, maxUnfilteredResults, null, sortOrder, cancellationToken))
 							{
 								CachedCommitDoc cachedCommit = await CachedCommitDoc.FromCommitAsync(commit, cancellationToken);
 								cachedCommit.PostLoad(_owner, StreamConfig);
@@ -775,7 +780,7 @@ namespace HordeServer.VersionControl.Perforce
 				{
 					_owner._logger.LogDebug("Querying Perforce server for {StreamId} commits from {MinChange} to {MaxChange} (max: {MaxResults}, tags: {Tags})", StreamConfig.Id, minChange ?? -2, maxChange ?? -2, maxResults ?? -1, (tags == null || tags.Count == 0) ? "none" : String.Join("/", tags.Select(x => x.ToString())));
 
-					await foreach (ICommit commit in base.FindAsync(minChange, true, maxChange, true, maxResults, tags, cancellationToken))
+					await foreach (ICommit commit in base.FindAsync(minChange, true, maxChange, true, maxResults, tags, sortOrder, cancellationToken))
 					{
 						yield return commit;
 					}
diff --git a/Engine/Source/Programs/Shared/EpicGames.Horde/Commits/ICommitCollection.cs b/Engine/Source/Programs/Shared/EpicGames.Horde/Commits/ICommitCollection.cs
--- a/Engine/Source/Programs/Shared/EpicGames.Horde/Commits/ICommitCollection.cs
+++ b/Engine/Source/Programs/Shared/EpicGames.Horde/Commits/ICommitCollection.cs
@@ -9,6 +9,22 @@ using EpicGames.Horde.Jobs.Templates;
 
 namespace EpicGames.Horde.Commits
 {
+	/// <summary>
+	/// Sort order for commit queries
+	/// </summary>
+	public enum CommitSortOrder
+	{
+		/// <summary>
+		/// Newest commits first (default)
+		/// </summary>
+		Descending,
+
+		/// <summary>
+		/// Oldest commits first
+		/// </summary>
+		Ascending
+	}
+
 	/// <summary>
 	/// Represents errors that occur during commit collection operations.
 	/// This exception is thrown when operations on a <see cref="ICommitCollection"/> fail.
@@ -48,7 +64,7 @@ namespace EpicGames.Horde.Commits
 		ValueTask<CommitIdWithOrder> GetOrderedAsync(CommitId commitId, CancellationToken cancellationToken = default);
 
 		/// <summary>
-		/// Finds changes submitted to a stream, in reverse order.
+		/// Finds changes submitted to a stream.
 		/// </summary>
 		/// <param name="minCommitId">The minimum changelist number</param>
 		/// <param name="includeMinCommit">Whether to include the minimum changelist in the range of enumerated responses</param>
@@ -56,9 +72,10 @@ namespace EpicGames.Horde.Commits
 		/// <param name="includeMaxCommit">Whether to include the maximum changelist in the range of enumerated responses</param>
 		/// <param name="maxResults">Maximum number of results to return</param>
 		/// <param name="tags">Tags for the commits to return</param>
+		/// <param name="sortOrder">Sort order for results (Descending = newest first, Ascending = oldest first)</param>
 		/// <param name="cancellationToken">Cancellation token for the operation</param>
 		/// <returns>Changelist information</returns>
-		IAsyncEnumerable<ICommit> FindAsync(CommitId? minCommitId = null, bool includeMinCommit = true, CommitId? maxCommitId = null, bool includeMaxCommit = true, int? maxResults = null, IReadOnlyList<CommitTag>? tags = null, CancellationToken cancellationToken = default);
+		IAsyncEnumerable<ICommit> FindAsync(CommitId? minCommitId = null, bool includeMinCommit = true, CommitId? maxCommitId = null, bool includeMaxCommit = true, int? maxResults = null, IReadOnlyList<CommitTag>? tags = null, CommitSortOrder sortOrder = CommitSortOrder.Descending, CancellationToken cancellationToken = default);
 
 		/// <summary>
 		/// Subscribes to changes from this commit source
