Horde: Add memory tracking for child processes during lease execution

Implements WindowsProcessMemoryTracker to sample memory usage of processes within a job object at regular intervals, tracking peak private and working set memory. Logs both Horde Agent's own memory usage and summarizes top memory-consuming child processes at lease completion.

Windows only for now.

diff --git a/Engine/Source/Programs/Horde/HordeAgent.Tests/Utility/WindowsProcessMemoryTrackerTests.cs b/Engine/Source/Programs/Horde/HordeAgent.Tests/Utility/WindowsProcessMemoryTrackerTests.cs
new file mode 100644
--- /dev/null
+++ b/Engine/Source/Programs/Horde/HordeAgent.Tests/Utility/WindowsProcessMemoryTrackerTests.cs
@@ -0,0 +1,351 @@
+// Copyright Epic Games, Inc. All Rights Reserved.
+
+using System;
+using System.Collections.Generic;
+using System.Diagnostics;
+using System.Runtime.Versioning;
+using System.Threading;
+using System.Threading.Tasks;
+using EpicGames.Core;
+using HordeAgent.Utility;
+using Microsoft.Extensions.Logging;
+using Microsoft.VisualStudio.TestTools.UnitTesting;
+
+namespace HordeAgent.Tests.Utility;
+
+[TestClass]
+[SupportedOSPlatform("windows")]
+public class WindowsProcessMemoryTrackerTests
+{
+	private ILogger _logger = null!;
+
+	[TestInitialize]
+	public void Setup()
+	{
+		using ILoggerFactory loggerFactory = LoggerFactory.Create(builder =>
+		{
+			builder.SetMinimumLevel(LogLevel.Debug);
+			builder.AddSimpleConsole(options => { options.SingleLine = true; });
+		});
+
+		_logger = loggerFactory.CreateLogger("WindowsProcessMemoryTrackerTests");
+	}
+
+	[TestMethod]
+	[TestCategory("Windows")]
+	public async Task TrackerStartsAndStopsCleanlyAsync()
+	{
+		if (!OperatingSystem.IsWindows())
+		{
+			Assert.Inconclusive("This test only runs on Windows");
+			return;
+		}
+
+		using ManagedProcessGroup processGroup = new();
+		await using WindowsProcessMemoryTracker tracker = new(processGroup, TimeSpan.FromMilliseconds(100), _logger);
+
+		tracker.Start();
+		await Task.Delay(150); // Let it run for a bit
+
+		// Should dispose cleanly without throwing
+	}
+
+	[TestMethod]
+	[TestCategory("Windows")]
+	public async Task TrackerReturnsEmptyListWhenNoProcessesAsync()
+	{
+		if (!OperatingSystem.IsWindows())
+		{
+			Assert.Inconclusive("This test only runs on Windows");
+			return;
+		}
+
+		using ManagedProcessGroup processGroup = new();
+
+		await using WindowsProcessMemoryTracker tracker = new(processGroup, TimeSpan.FromMilliseconds(100), _logger);
+		tracker.Start();
+		await Task.Delay(50); // Give it time to sample
+
+		IReadOnlyList<ProcessMemoryInfo> topProcesses = tracker.GetTopMemoryProcesses(5);
+
+		// Should return empty list, not throw
+		Assert.IsNotNull(topProcesses);
+	}
+
+	[TestMethod]
+	[TestCategory("Windows")]
+	public async Task TrackerTracksChildProcessMemoryAsync()
+	{
+		if (!OperatingSystem.IsWindows())
+		{
+			Assert.Inconclusive("This test only runs on Windows");
+			return;
+		}
+
+		using ManagedProcessGroup processGroup = new();
+
+		// Start a simple process that allocates some memory
+		using ManagedProcess process = new(
+			processGroup,
+			"cmd.exe",
+			"/c ping -n 2 127.0.0.1",
+			null,
+			null,
+			ProcessPriorityClass.Normal);
+
+		await using WindowsProcessMemoryTracker tracker = new(
+			processGroup,
+			TimeSpan.FromMilliseconds(50),
+			_logger);
+
+		tracker.Start();
+
+		// Wait for process to run and be sampled
+		await Task.Delay(200);
+
+		IReadOnlyList<ProcessMemoryInfo> topProcesses = tracker.GetTopMemoryProcesses(5);
+
+		// Should have tracked at least the cmd.exe process
+		Assert.IsTrue(topProcesses.Count > 0, "Expected at least one tracked process");
+
+		// Verify the tracked process has reasonable data
+		ProcessMemoryInfo firstProcess = topProcesses[0];
+		Assert.IsTrue(firstProcess.ProcessId > 0, "ProcessId should be positive");
+		Assert.IsFalse(string.IsNullOrEmpty(firstProcess.ProcessName), "ProcessName should not be empty");
+		Assert.IsTrue(firstProcess.PeakPrivateBytes > 0, "PeakPrivateBytes should be positive");
+		Assert.IsTrue(firstProcess.PeakWorkingSetBytes > 0, "PeakWorkingSetBytes should be positive");
+
+		// Wait for process to finish
+		await process.WaitForExitAsync(CancellationToken.None);
+	}
+
+	[TestMethod]
+	[TestCategory("Windows")]
+	public async Task TrackerUpdatesPeakMemoryOverTimeAsync()
+	{
+		if (!OperatingSystem.IsWindows())
+		{
+			Assert.Inconclusive("This test only runs on Windows");
+			return;
+		}
+
+		using ManagedProcessGroup processGroup = new();
+
+		// Start a longer-running process
+		using ManagedProcess process = new(
+			processGroup,
+			"cmd.exe",
+			"/c ping -n 3 127.0.0.1",
+			null,
+			null,
+			ProcessPriorityClass.Normal);
+
+		await using WindowsProcessMemoryTracker tracker = new(
+			processGroup,
+			TimeSpan.FromMilliseconds(50),
+			_logger);
+
+		tracker.Start();
+
+		// Wait for first sample
+		await Task.Delay(100);
+
+		IReadOnlyList<ProcessMemoryInfo> firstSample = tracker.GetTopMemoryProcesses(5);
+		Assert.IsTrue(firstSample.Count > 0, "Expected at least one tracked process after first sample");
+
+		// Wait for more samples
+		await Task.Delay(200);
+
+		IReadOnlyList<ProcessMemoryInfo> laterSample = tracker.GetTopMemoryProcesses(5);
+		Assert.IsTrue(laterSample.Count > 0, "Expected at least one tracked process after later sample");
+
+		// Peak values should be maintained (not reset)
+		// Note: We can't guarantee peak increased, but it should at least not be zero
+		Assert.IsTrue(laterSample[0].PeakPrivateBytes > 0);
+		Assert.IsTrue(laterSample[0].PeakWorkingSetBytes > 0);
+
+		// Wait for process to finish
+		await process.WaitForExitAsync(CancellationToken.None);
+	}
+
+	[TestMethod]
+	[TestCategory("Windows")]
+	public async Task TrackerCapturesCommandLineAsync()
+	{
+		if (!OperatingSystem.IsWindows())
+		{
+			Assert.Inconclusive("This test only runs on Windows");
+			return;
+		}
+
+		using ManagedProcessGroup processGroup = new();
+
+		// Start a process with a distinctive command line
+		string uniqueArg = $"test-arg-{Guid.NewGuid():N}";
+		using ManagedProcess process = new(
+			processGroup,
+			"cmd.exe",
+			$"/c echo {uniqueArg} && ping -n 2 127.0.0.1",
+			null,
+			null,
+			ProcessPriorityClass.Normal);
+
+		await using WindowsProcessMemoryTracker tracker = new(
+			processGroup,
+			TimeSpan.FromMilliseconds(50),
+			_logger);
+
+		tracker.Start();
+
+		// Wait for sampling
+		await Task.Delay(150);
+
+		IReadOnlyList<ProcessMemoryInfo> topProcesses = tracker.GetTopMemoryProcesses(10);
+
+		// Find cmd.exe in the list
+		ProcessMemoryInfo? cmdProcess = null;
+		foreach (ProcessMemoryInfo proc in topProcesses)
+		{
+			if (proc.ProcessName.Equals("cmd", StringComparison.OrdinalIgnoreCase))
+			{
+				cmdProcess = proc;
+				break;
+			}
+		}
+
+		Assert.IsNotNull(cmdProcess, "Expected to find cmd.exe in tracked processes");
+
+		// Command line should be captured (may be truncated)
+		Assert.IsFalse(string.IsNullOrEmpty(cmdProcess.CommandLine), "CommandLine should not be empty");
+
+		// Wait for process to finish
+		await process.WaitForExitAsync(CancellationToken.None);
+	}
+
+	[TestMethod]
+	[TestCategory("Windows")]
+	public void TrackerTruncatesLongCommandLines()
+	{
+		if (!OperatingSystem.IsWindows())
+		{
+			Assert.Inconclusive("This test only runs on Windows");
+			return;
+		}
+
+		// Verify the MaxCommandLineLength constant is reasonable
+		Assert.AreEqual(100, WindowsProcessMemoryTracker.MaxCommandLineLength);
+	}
+
+	[TestMethod]
+	[TestCategory("Windows")]
+	public async Task GetTopMemoryProcessesReturnsCorrectCountAsync()
+	{
+		if (!OperatingSystem.IsWindows())
+		{
+			Assert.Inconclusive("This test only runs on Windows");
+			return;
+		}
+
+		using ManagedProcessGroup processGroup = new();
+
+		// Start multiple processes
+		using ManagedProcess process1 = new(
+			processGroup,
+			"cmd.exe",
+			"/c ping -n 2 127.0.0.1",
+			null,
+			null,
+			ProcessPriorityClass.Normal);
+
+		using ManagedProcess process2 = new(
+			processGroup,
+			"cmd.exe",
+			"/c ping -n 2 127.0.0.1",
+			null,
+			null,
+			ProcessPriorityClass.Normal);
+
+		await using WindowsProcessMemoryTracker tracker = new(
+			processGroup,
+			TimeSpan.FromMilliseconds(50),
+			_logger);
+
+		tracker.Start();
+
+		// Wait for sampling
+		await Task.Delay(150);
+
+		// Request only 1 process
+		IReadOnlyList<ProcessMemoryInfo> top1 = tracker.GetTopMemoryProcesses(1);
+		Assert.IsTrue(top1.Count <= 1, "GetTopMemoryProcesses(1) should return at most 1 process");
+
+		// Request 5 processes
+		IReadOnlyList<ProcessMemoryInfo> top5 = tracker.GetTopMemoryProcesses(5);
+		Assert.IsTrue(top5.Count <= 5, "GetTopMemoryProcesses(5) should return at most 5 processes");
+
+		// Verify processes are sorted by PeakPrivateBytes descending
+		for (int i = 1; i < top5.Count; i++)
+		{
+			Assert.IsTrue(top5[i - 1].PeakPrivateBytes >= top5[i].PeakPrivateBytes,
+				"Processes should be sorted by PeakPrivateBytes descending");
+		}
+
+		// Wait for processes to finish
+		await process1.WaitForExitAsync(CancellationToken.None);
+		await process2.WaitForExitAsync(CancellationToken.None);
+	}
+
+	[TestMethod]
+	[TestCategory("Windows")]
+	public async Task TrackerThrowsWhenStartedTwiceAsync()
+	{
+		if (!OperatingSystem.IsWindows())
+		{
+			Assert.Inconclusive("This test only runs on Windows");
+			return;
+		}
+
+		using ManagedProcessGroup processGroup = new();
+
+		await using WindowsProcessMemoryTracker tracker = new(
+			processGroup,
+			TimeSpan.FromMilliseconds(100),
+			_logger);
+
+		tracker.Start();
+
+		Assert.ThrowsException<InvalidOperationException>(() => tracker.Start());
+	}
+
+	[TestMethod]
+	[TestCategory("Windows")]
+	public void TrackerThrowsOnNullProcessGroup()
+	{
+		if (!OperatingSystem.IsWindows())
+		{
+			Assert.Inconclusive("This test only runs on Windows");
+			return;
+		}
+
+		// Constructor throws before any resources are allocated, so no dispose needed
+		Assert.ThrowsException<ArgumentNullException>(() =>
+			_ = new WindowsProcessMemoryTracker(null!, TimeSpan.FromSeconds(30), _logger));
+	}
+
+	[TestMethod]
+	[TestCategory("Windows")]
+	public void TrackerThrowsOnNullLogger()
+	{
+		if (!OperatingSystem.IsWindows())
+		{
+			Assert.Inconclusive("This test only runs on Windows");
+			return;
+		}
+
+		using ManagedProcessGroup processGroup = new();
+
+		// Constructor throws before any resources are allocated, so no dispose needed
+		Assert.ThrowsException<ArgumentNullException>(() =>
+			_ = new WindowsProcessMemoryTracker(processGroup, TimeSpan.FromSeconds(30), null!));
+	}
+}
diff --git a/Engine/Source/Programs/Horde/HordeAgent/Leases/LeaseHandler.cs b/Engine/Source/Programs/Horde/HordeAgent/Leases/LeaseHandler.cs
--- a/Engine/Source/Programs/Horde/HordeAgent/Leases/LeaseHandler.cs
+++ b/Engine/Source/Programs/Horde/HordeAgent/Leases/LeaseHandler.cs
@@ -7,6 +7,7 @@ using EpicGames.Horde.Agents.Metrics;
 using Google.Protobuf;
 using Google.Protobuf.WellKnownTypes;
 using HordeAgent.Services;
+using HordeAgent.Utility;
 using HordeCommon.Rpc;
 using HordeCommon.Rpc.Messages;
 using Microsoft.Extensions.Logging;
@@ -20,6 +21,11 @@ namespace HordeAgent.Leases
 	/// </summary>
 	abstract class LeaseHandler : IDisposable
 	{
+		/// <summary>
+		/// Conversion factor from bytes to megabytes
+		/// </summary>
+		const double BytesPerMegabyte = 1024.0 * 1024.0;
+
 		/// <summary>
 		/// Identifier for this lease
 		/// </summary>
@@ -123,7 +129,14 @@ namespace HordeAgent.Leases
 
 			Stopwatch leaseDuration = Stopwatch.StartNew();
 			CpuTimeMetrics? cpuTimeStart = systemMetrics.GetCpuTime();
-			
+
+			// Track Horde agent's own memory at lease start (Windows only)
+			long agentMemoryStart = 0;
+			if (OperatingSystem.IsWindows())
+			{
+				agentMemoryStart = Process.GetCurrentProcess().PrivateMemorySize64;
+			}
+
 			// Get the lease outcome
 			LeaseResult result = new (LeaseOutcome.Failed, LeaseOutcomeReason.UnexpectedError);
 			try
@@ -143,6 +156,12 @@ namespace HordeAgent.Leases
 			}
 			CpuTimeMetrics? cpuTimeEnd = systemMetrics.GetCpuTime();
 
+			// Log HordeAgent's own memory delta (Windows only)
+			if (OperatingSystem.IsWindows())
+			{
+				LogAgentMemory(agentMemoryStart, logger);
+			}
+
 			// Update the state of the lease
 			RpcLease newRpcLease = new RpcLease(RpcLease);
 			if (_cancellationSource.IsCancellationRequested)
@@ -202,30 +221,17 @@ namespace HordeAgent.Leases
 			logger.LogInformation("Running child process: {Executable} {CommandLine}", CommandLineArguments.Quote(executable), commandLine);
 			try
 			{
-				using (ManagedProcessGroup processGroup = new ManagedProcessGroup())
-				using (ManagedProcess process = new ManagedProcess(processGroup, executable, commandLine, null, environment, ProcessPriorityClass.Normal))
+				using ManagedProcessGroup processGroup = new ManagedProcessGroup();
+				using ManagedProcess process = new ManagedProcess(processGroup, executable, commandLine, null, environment, ProcessPriorityClass.Normal);
+
+				// Run the process with optional memory tracking
+				if (OperatingSystem.IsWindows())
 				{
-					for (; ; )
-					{
-						string? line = await process.ReadLineAsync(cancellationToken);
-						if (line == null)
-						{
-							break;
-						}
-
-						JsonLogEvent jsonLogEvent;
-						if (JsonLogEvent.TryParse(line, out jsonLogEvent))
-						{
-							logger.LogJsonLogEvent(jsonLogEvent);
-						}
-						else
-						{
-							logger.LogInformation("{Line}", line);
-						}
-					}
-
-					await process.WaitForExitAsync(CancellationToken.None);
-					return process.ExitCode;
+					return await RunProcessWithMemoryTrackingAsync(processGroup, process, logger, cancellationToken);
+				}
+				else
+				{
+					return await RunProcessCoreAsync(process, logger, cancellationToken);
 				}
 			}
 			catch (Exception ex)
@@ -235,6 +241,52 @@ namespace HordeAgent.Leases
 			}
 		}
 
+		/// <summary>
+		/// Runs the process with memory tracking (Windows only)
+		/// </summary>
+		[System.Runtime.Versioning.SupportedOSPlatform("windows")]
+		private static async Task<int> RunProcessWithMemoryTrackingAsync(ManagedProcessGroup processGroup, ManagedProcess process, ILogger logger, CancellationToken cancellationToken)
+		{
+			await using WindowsProcessMemoryTracker memoryTracker = new (processGroup, TimeSpan.FromSeconds(30), logger);
+			memoryTracker.Start();
+
+			try
+			{
+				return await RunProcessCoreAsync(process, logger, cancellationToken);
+			}
+			finally
+			{
+				LogChildProcessMemory(processGroup, memoryTracker, logger);
+			}
+		}
+
+		/// <summary>
+		/// Core process execution loop
+		/// </summary>
+		private static async Task<int> RunProcessCoreAsync(ManagedProcess process, ILogger logger, CancellationToken cancellationToken)
+		{
+			for (; ; )
+			{
+				string? line = await process.ReadLineAsync(cancellationToken);
+				if (line == null)
+				{
+					break;
+				}
+
+				if (JsonLogEvent.TryParse(line, out JsonLogEvent jsonLogEvent))
+				{
+					logger.LogJsonLogEvent(jsonLogEvent);
+				}
+				else
+				{
+					logger.LogInformation("{Line}", line);
+				}
+			}
+
+			await process.WaitForExitAsync(CancellationToken.None);
+			return process.ExitCode;
+		}
+
 		/// <summary>
 		/// Runs a .NET assembly as a child process, piping the output to the given logger
 		/// </summary>
@@ -295,10 +347,62 @@ namespace HordeAgent.Leases
 			catch (Exception e)
 			{
 				logger.LogWarning(e, "Failed to calculate lease metrics");
-				
 			}
 			return metrics;
 		}
+
+		/// <summary>
+		/// Logs HordeAgent's own memory usage at lease end
+		/// </summary>
+		private static void LogAgentMemory(long memoryStart, ILogger logger)
+		{
+			using Process currentProcess = Process.GetCurrentProcess();
+			long memoryEnd = currentProcess.PrivateMemorySize64;
+			long peakMemory = currentProcess.PeakWorkingSet64;
+
+			logger.LogInformation(
+				"HordeAgent Memory: Start={StartMb:F1} MB End={EndMb:F1} MB Delta={DeltaMb:+0.0;-0.0;0.0} MB Peak={PeakMb:F1} MB",
+				memoryStart / BytesPerMegabyte,
+				memoryEnd / BytesPerMegabyte,
+				(memoryEnd - memoryStart) / BytesPerMegabyte,
+				peakMemory / BytesPerMegabyte);
+		}
+
+		/// <summary>
+		/// Logs child process memory summary
+		/// </summary>
+		[System.Runtime.Versioning.SupportedOSPlatform("windows")]
+		private static void LogChildProcessMemory(
+			ManagedProcessGroup processGroup,
+			WindowsProcessMemoryTracker tracker,
+			ILogger logger)
+		{
+			IReadOnlyList<ProcessMemoryInfo> topProcesses = tracker.GetTopMemoryProcesses(5);
+
+			// Log job-level aggregates from ManagedProcessGroup
+			long totalPeakBytes = processGroup.PeakJobMemoryUsed;
+			long largestProcessPeakBytes = processGroup.PeakProcessMemoryUsed;
+
+			logger.LogInformation(
+				"Child Process Memory Summary: TotalPeak={TotalPeakMb:F1} MB, LargestSingleProcess={LargestMb:F1} MB",
+				totalPeakBytes / BytesPerMegabyte,
+				largestProcessPeakBytes / BytesPerMegabyte);
+
+			if (topProcesses.Count > 0)
+			{
+				logger.LogInformation("Top {Count} memory-consuming processes:", topProcesses.Count);
+				foreach (ProcessMemoryInfo proc in topProcesses)
+				{
+					logger.LogInformation(
+						"  PID={Pid} {Name}.exe PeakPrivate={PrivateMb:F1} MB PeakWorkingSet={WorkingSetMb:F1} MB Args=\"{Args}\"",
+						proc.ProcessId,
+						proc.ProcessName,
+						proc.PeakPrivateBytes / BytesPerMegabyte,
+						proc.PeakWorkingSetBytes / BytesPerMegabyte,
+						proc.CommandLine);
+				}
+			}
+		}
 	}
 
 	/// <summary>
diff --git a/Engine/Source/Programs/Horde/HordeAgent/Utility/WindowsProcessMemoryTracker.cs b/Engine/Source/Programs/Horde/HordeAgent/Utility/WindowsProcessMemoryTracker.cs
new file mode 100644
--- /dev/null
+++ b/Engine/Source/Programs/Horde/HordeAgent/Utility/WindowsProcessMemoryTracker.cs
@@ -0,0 +1,299 @@
+// Copyright Epic Games, Inc. All Rights Reserved.
+
+using System.Collections.Concurrent;
+using System.ComponentModel;
+using System.Diagnostics;
+using System.Runtime.InteropServices;
+using System.Runtime.Versioning;
+using EpicGames.Core;
+using Microsoft.Extensions.Logging;
+using Microsoft.Win32.SafeHandles;
+
+namespace HordeAgent.Utility
+{
+	/// <summary>
+	/// Memory information for a process, tracking peak values observed across all samples
+	/// </summary>
+	/// <param name="ProcessId">The process ID</param>
+	/// <param name="ProcessName">Name of the process executable (without extension)</param>
+	/// <param name="CommandLine">Command line arguments (truncated to MaxCommandLineLength)</param>
+	/// <param name="PeakPrivateBytes">Maximum PrivateMemorySize64 observed - best for leak detection</param>
+	/// <param name="PeakWorkingSetBytes">Maximum WorkingSet64 observed - RAM impact</param>
+	[SupportedOSPlatform("windows")]
+	public record ProcessMemoryInfo(
+		int ProcessId,
+		string ProcessName,
+		string CommandLine,
+		long PeakPrivateBytes,
+		long PeakWorkingSetBytes);
+
+	/// <summary>
+	/// Tracks memory usage for all processes in a Windows Job Object.
+	/// Samples at configurable intervals and maintains peak memory values per process.
+	/// </summary>
+	[SupportedOSPlatform("windows")]
+	public sealed class WindowsProcessMemoryTracker : IAsyncDisposable
+	{
+		/// <summary>
+		/// Maximum length for stored command line strings
+		/// </summary>
+		public const int MaxCommandLineLength = 100;
+
+		#region Native Interop
+
+#pragma warning disable IDE1006 // Naming Styles
+#pragma warning disable CS0649 // Field never assigned
+
+		const int JobObjectBasicProcessIdList = 3;
+
+		[StructLayout(LayoutKind.Sequential)]
+		struct JOBOBJECT_BASIC_PROCESS_ID_LIST
+		{
+			public uint NumberOfAssignedProcesses;
+			public uint NumberOfProcessIdsInList;
+			// Followed by array of UIntPtr ProcessIdList[1]
+		}
+
+		[DllImport("kernel32.dll", SetLastError = true)]
+		static extern bool QueryInformationJobObject(
+			SafeFileHandle hJob,
+			int JobObjectInformationClass,
+			IntPtr lpJobObjectInfo,
+			int cbJobObjectInfoLength,
+			out int lpReturnLength);
+
+#pragma warning restore CS0649
+#pragma warning restore IDE1006
+
+		#endregion
+
+		readonly ManagedProcessGroup _processGroup;
+		readonly TimeSpan _sampleInterval;
+		readonly ILogger _logger;
+
+		readonly ConcurrentDictionary<int, ProcessMemoryInfo> _peakByPid = new();
+		Task? _samplingTask;
+		CancellationTokenSource? _cts;
+
+		/// <summary>
+		/// Creates a new memory tracker for the specified process group
+		/// </summary>
+		/// <param name="processGroup">The managed process group containing the job object to monitor</param>
+		/// <param name="sampleInterval">How often to sample process memory (e.g., 30 seconds)</param>
+		/// <param name="logger">Logger for diagnostic messages</param>
+		public WindowsProcessMemoryTracker(
+			ManagedProcessGroup processGroup,
+			TimeSpan sampleInterval,
+			ILogger logger)
+		{
+			_processGroup = processGroup ?? throw new ArgumentNullException(nameof(processGroup));
+			_sampleInterval = sampleInterval;
+			_logger = logger ?? throw new ArgumentNullException(nameof(logger));
+		}
+
+		/// <summary>
+		/// Starts the background memory sampling task
+		/// </summary>
+		public void Start()
+		{
+			if (_samplingTask != null)
+			{
+				throw new InvalidOperationException("Memory tracker is already running");
+			}
+
+			_cts = new CancellationTokenSource();
+			_samplingTask = SamplingLoopAsync(_cts.Token);
+		}
+
+		/// <summary>
+		/// Returns the top N processes by peak private memory
+		/// </summary>
+		/// <param name="count">Number of processes to return (default 5)</param>
+		/// <returns>List of process memory info sorted by peak private bytes descending</returns>
+		public IReadOnlyList<ProcessMemoryInfo> GetTopMemoryProcesses(int count = 5)
+		{
+			return _peakByPid.Values
+				.OrderByDescending(p => p.PeakPrivateBytes)
+				.Take(count)
+				.ToList();
+		}
+
+		/// <inheritdoc/>
+		public async ValueTask DisposeAsync()
+		{
+			if (_cts != null)
+			{
+				await _cts.CancelAsync();
+
+				if (_samplingTask != null)
+				{
+					try
+					{
+						await _samplingTask;
+					}
+					catch (OperationCanceledException)
+					{
+						// Expected
+					}
+				}
+
+				_cts.Dispose();
+				_cts = null;
+			}
+		}
+
+		async Task SamplingLoopAsync(CancellationToken cancellationToken)
+		{
+			// Do an immediate sample at start
+			SampleProcessMemory();
+
+			using PeriodicTimer timer = new(_sampleInterval);
+			while (await timer.WaitForNextTickAsync(cancellationToken))
+			{
+				SampleProcessMemory();
+			}
+		}
+
+		void SampleProcessMemory()
+		{
+			try
+			{
+				int[] pids = GetJobProcessIds();
+				foreach (int pid in pids)
+				{
+					SampleProcess(pid);
+				}
+			}
+			catch (Exception ex)
+			{
+				_logger.LogDebug(ex, "Error during memory sampling");
+			}
+		}
+
+		int[] GetJobProcessIds()
+		{
+			SafeFileHandle? jobHandle = _processGroup.JobHandle;
+			if (jobHandle == null || jobHandle.IsInvalid || jobHandle.IsClosed)
+			{
+				return [];
+			}
+
+			// Start with space for 64 processes, grow if needed
+			int bufferSize = Marshal.SizeOf<JOBOBJECT_BASIC_PROCESS_ID_LIST>() + (64 * IntPtr.Size);
+			IntPtr buffer = Marshal.AllocHGlobal(bufferSize);
+
+			try
+			{
+				while (true)
+				{
+					if (QueryInformationJobObject(
+						jobHandle,
+						JobObjectBasicProcessIdList,
+						buffer,
+						bufferSize,
+						out int returnLength))
+					{
+						JOBOBJECT_BASIC_PROCESS_ID_LIST info = Marshal.PtrToStructure<JOBOBJECT_BASIC_PROCESS_ID_LIST>(buffer);
+
+						// Check if we need more space
+						if (info.NumberOfAssignedProcesses > info.NumberOfProcessIdsInList)
+						{
+							bufferSize = Marshal.SizeOf<JOBOBJECT_BASIC_PROCESS_ID_LIST>() + ((int)info.NumberOfAssignedProcesses * IntPtr.Size);
+							IntPtr newBuffer = Marshal.AllocHGlobal(bufferSize);
+							Marshal.FreeHGlobal(buffer);
+							buffer = newBuffer;
+							continue;
+						}
+
+						// Read the PIDs
+						int[] pids = new int[info.NumberOfProcessIdsInList];
+						IntPtr pidArrayStart = buffer + Marshal.SizeOf<JOBOBJECT_BASIC_PROCESS_ID_LIST>();
+
+						for (int i = 0; i < info.NumberOfProcessIdsInList; i++)
+						{
+							IntPtr pidPtr = Marshal.ReadIntPtr(pidArrayStart, i * IntPtr.Size);
+							pids[i] = (int)pidPtr.ToInt64();
+						}
+
+						return pids;
+					}
+					else
+					{
+						int error = Marshal.GetLastWin32Error();
+						// ERROR_MORE_DATA = 234
+						if (error == 234)
+						{
+							bufferSize *= 2;
+							IntPtr newBuffer = Marshal.AllocHGlobal(bufferSize);
+							Marshal.FreeHGlobal(buffer);
+							buffer = newBuffer;
+							continue;
+						}
+
+						_logger.LogDebug("QueryInformationJobObject failed with error {Error}", error);
+						return [];
+					}
+				}
+			}
+			finally
+			{
+				Marshal.FreeHGlobal(buffer);
+			}
+		}
+
+		void SampleProcess(int pid)
+		{
+			try
+			{
+				using Process process = Process.GetProcessById(pid);
+
+				long privateBytes = process.PrivateMemorySize64;
+				long workingSet = process.WorkingSet64;
+				string processName = process.ProcessName;
+				string commandLine = TruncateCommandLine(ProcessUtils.GetProcessCommandLine(pid) ?? string.Empty);
+
+				_peakByPid.AddOrUpdate(pid,
+					// Add new entry
+					_ => new ProcessMemoryInfo(
+						pid,
+						processName,
+						commandLine,
+						privateBytes,
+						workingSet),
+					// Update existing - keep max values
+					(_, existing) => existing with
+					{
+						PeakPrivateBytes = Math.Max(existing.PeakPrivateBytes, privateBytes),
+						PeakWorkingSetBytes = Math.Max(existing.PeakWorkingSetBytes, workingSet)
+					});
+			}
+			catch (ArgumentException)
+			{
+				// Process exited
+			}
+			catch (InvalidOperationException)
+			{
+				// Process exited during query
+			}
+			catch (Win32Exception)
+			{
+				// Access denied
+			}
+		}
+
+		static string TruncateCommandLine(string commandLine)
+		{
+			if (string.IsNullOrEmpty(commandLine))
+			{
+				return string.Empty;
+			}
+
+			if (commandLine.Length <= MaxCommandLineLength)
+			{
+				return commandLine;
+			}
+
+			return commandLine[..(MaxCommandLineLength - 3)] + "...";
+		}
+	}
+}
diff --git a/Engine/Source/Programs/Shared/EpicGames.Core/ManagedProcess.cs b/Engine/Source/Programs/Shared/EpicGames.Core/ManagedProcess.cs
--- a/Engine/Source/Programs/Shared/EpicGames.Core/ManagedProcess.cs
+++ b/Engine/Source/Programs/Shared/EpicGames.Core/ManagedProcess.cs
@@ -122,7 +122,7 @@ namespace EpicGames.Core
 		/// Handle to the native job object that this process is added to. This handle is closed by the Dispose() method (and will automatically be closed by the OS on process exit),
 		/// resulting in the child process being killed.
 		/// </summary>
-		internal SafeFileHandle? JobHandle
+		public SafeFileHandle? JobHandle
 		{
 			get;
 			private set;
