Horde: Add IRcaProcessor framework and config options.
Enables licensees to implement their own root cause analysis pipelines by introducing an extensible IRcaProcessor interface in HordeServer.Shared.

diff --git a/Engine/Source/Programs/Horde/HordeServer.Shared/RCA/IRcaProcessor.cs b/Engine/Source/Programs/Horde/HordeServer.Shared/RCA/IRcaProcessor.cs
new file mode 100644
--- /dev/null
+++ b/Engine/Source/Programs/Horde/HordeServer.Shared/RCA/IRcaProcessor.cs
@@ -0,0 +1,24 @@
+// Copyright Epic Games, Inc. All Rights Reserved.
+
+using EpicGames.Horde.Issues;
+using EpicGames.Horde.Streams;
+
+namespace HordeServer.RCA
+{
+	/// <summary>
+	/// Interface for processing Root Cause Analysis on issues.
+	/// Implementations may use LLMs, ML models, or other techniques.
+	/// </summary>
+	public interface IRcaProcessor
+	{
+		/// <summary>
+		/// Process RCA for the given issue if eligible. Implementations should
+		/// check workflow eligibility internally and return early if not eligible.
+		/// </summary>
+		/// <param name="issueId">The issue ID to analyze</param>
+		/// <param name="streamId">The stream the issue belongs to</param>
+		/// <param name="workflowId">The workflow ID (null if not in a workflow)</param>
+		/// <param name="cancellationToken">Cancellation token</param>
+		Task ProcessRcaAsync(int issueId, StreamId streamId, WorkflowId? workflowId, CancellationToken cancellationToken);
+	}
+}
diff --git a/Engine/Source/Programs/Horde/Plugins/Build/HordeServer.Build/BuildConfig.cs b/Engine/Source/Programs/Horde/Plugins/Build/HordeServer.Build/BuildConfig.cs
--- a/Engine/Source/Programs/Horde/Plugins/Build/HordeServer.Build/BuildConfig.cs
+++ b/Engine/Source/Programs/Horde/Plugins/Build/HordeServer.Build/BuildConfig.cs
@@ -105,6 +105,17 @@ namespace HordeServer
 		/// </summary>
 		public bool EnableIssueRootCauseAnalysisNotifications { get; set; }
 
+		/// <summary>
+		/// Whether to use multi-stage LLM analysis for auto-RCA. Multi-stage is slower but potentially more accurate. Set to false for faster processing with higher throughput.
+		/// </summary>
+		public bool AutoRcaUseMultiStage { get; set; } = true;
+
+		/// <summary>
+		/// Maximum number of concurrent auto-RCA requests. Limits load on the external LLM API.
+		/// Set to 0 to disable auto-RCA globally (acts as a kill switch).
+		/// </summary>
+		public int MaxConcurrentRcaRequests { get; set; }
+
 		/// <summary>
 		/// Issue channel groupings
 		/// </summary>
diff --git a/Engine/Source/Programs/Horde/Plugins/Build/HordeServer.Build/Issues/IssueService.cs b/Engine/Source/Programs/Horde/Plugins/Build/HordeServer.Build/Issues/IssueService.cs
--- a/Engine/Source/Programs/Horde/Plugins/Build/HordeServer.Build/Issues/IssueService.cs
+++ b/Engine/Source/Programs/Horde/Plugins/Build/HordeServer.Build/Issues/IssueService.cs
@@ -15,6 +15,7 @@ using EpicGames.Horde.Users;
 using HordeServer.Commits;
 using HordeServer.Jobs;
 using HordeServer.Logs;
+using HordeServer.RCA;
 using HordeServer.Streams;
 using HordeServer.Users;
 using Microsoft.Extensions.DependencyInjection;
@@ -248,11 +249,13 @@ namespace HordeServer.Issues
 		public IReadOnlyList<IIssueDetails> CachedOpenIssues => _cachedIssues;
 
 		readonly IOptionsMonitor<BuildConfig> _buildConfig;
+		readonly IRcaProcessor? _rcaProcessor;
+		readonly IHostApplicationLifetime _lifetime;
 
 		/// <summary>
 		/// Constructor
 		/// </summary>
-		public IssueService(JobService jobService, IIssueCollection issueCollection, ICommitService commitService, IJobStepRefCollection jobStepRefs, IUserCollection userCollection, ILogCollection logCollection, IClock clock, IOptionsMonitor<BuildConfig> buildConfig, Tracer tracer, ILogger<IssueService> logger)
+		public IssueService(JobService jobService, IIssueCollection issueCollection, ICommitService commitService, IJobStepRefCollection jobStepRefs, IUserCollection userCollection, ILogCollection logCollection, IClock clock, IOptionsMonitor<BuildConfig> buildConfig, Tracer tracer, ILogger<IssueService> logger, IHostApplicationLifetime lifetime, IRcaProcessor? rcaProcessor = null)
 		{
 			Type[] issueTypes = Assembly.GetExecutingAssembly().GetTypes().Where(x => !x.IsAbstract && typeof(IIssue).IsAssignableFrom(x)).ToArray();
 			foreach (Type issueType in issueTypes)
@@ -274,6 +277,8 @@ namespace HordeServer.Issues
 			_buildConfig = buildConfig;
 			_tracer = tracer;
 			_logger = logger;
+			_lifetime = lifetime;
+			_rcaProcessor = rcaProcessor;
 
 			// Find all the issue handler types
 			List<(Type Type, IssueHandlerAttribute Attribute)> handlerTypes = new List<(Type, IssueHandlerAttribute)>();
@@ -1826,6 +1831,30 @@ namespace HordeServer.Issues
 				string summary = GetSummary(span.Fingerprint, span.FirstFailure.Severity);
 				issue = await _issueCollection.AddIssueAsync(summary, cancellationToken);
 				LoggerUtils.MultiLog([(_logger, LogLevel.Information), (issueLogger, LogLevel.Information)], "Created issue {IssueId}", issue.Id);
+
+				// Fire-and-forget RCA processing for genuinely new issues only
+				// Processor handles workflow eligibility and rate limiting internally
+				// Uses ApplicationStopping token to cancel in-flight requests during shutdown
+				if (_rcaProcessor != null)
+				{
+					WorkflowId? workflowId = span.FirstFailure.Annotations?.WorkflowId;
+					CancellationToken stoppingToken = _lifetime.ApplicationStopping;
+					_ = Task.Run(async () =>
+					{
+						try
+						{
+							await _rcaProcessor.ProcessRcaAsync(issue.Id, span.StreamId, workflowId, stoppingToken);
+						}
+						catch (OperationCanceledException) when (stoppingToken.IsCancellationRequested)
+						{
+							_logger.LogDebug("RCA processing cancelled during shutdown for issue {IssueId}", issue.Id);
+						}
+						catch (Exception ex)
+						{
+							_logger.LogError(ex, "RCA processing failed for issue {IssueId}", issue.Id);
+						}
+					});
+				}
 			}
 			else
 			{
diff --git a/Engine/Source/Programs/Horde/Plugins/Build/HordeServer.Build/Issues/WorkflowConfig.cs b/Engine/Source/Programs/Horde/Plugins/Build/HordeServer.Build/Issues/WorkflowConfig.cs
--- a/Engine/Source/Programs/Horde/Plugins/Build/HordeServer.Build/Issues/WorkflowConfig.cs
+++ b/Engine/Source/Programs/Horde/Plugins/Build/HordeServer.Build/Issues/WorkflowConfig.cs
@@ -70,6 +70,9 @@ namespace HordeServer.Issues
 
 		IReadOnlyList<int> IWorkflow.EscalateRootCauseAnalysisTimes => EscalateRootCauseAnalysisTimes;
 
+		/// <inheritdoc/>
+		public bool AutoRcaEnabled { get; set; }
+
 		/// <inheritdoc/>
 		public int MaxMentions { get; set; } = 5;
 
diff --git a/Engine/Source/Programs/Shared/EpicGames.Horde/Issues/IWorkflow.cs b/Engine/Source/Programs/Shared/EpicGames.Horde/Issues/IWorkflow.cs
--- a/Engine/Source/Programs/Shared/EpicGames.Horde/Issues/IWorkflow.cs
+++ b/Engine/Source/Programs/Shared/EpicGames.Horde/Issues/IWorkflow.cs
@@ -92,10 +92,15 @@ namespace EpicGames.Horde.Issues
 		IReadOnlyList<int> EscalateTimes { get; }
 
 		/// <summary>
-		/// Times after a root cause analysis has been opened to escalate to the owner, in minutes. Continues to notify on the last interval once reaching the end of the list.
+		/// Times after a manual root cause analysis has been opened to escalate to the owner, in minutes. Continues to notify on the last interval once reaching the end of the list.
 		/// </summary>
 		IReadOnlyList<int> EscalateRootCauseAnalysisTimes { get; }
 
+		/// <summary>
+		/// Whether to automatically run RCA processing for issues in this workflow
+		/// </summary>
+		bool AutoRcaEnabled { get; }
+
 		/// <summary>
 		/// Maximum number of people to mention on a triage thread
 		/// </summary>
diff --git a/Engine/Source/Programs/Shared/EpicGames.Horde/Issues/WorkflowMessages.cs b/Engine/Source/Programs/Shared/EpicGames.Horde/Issues/WorkflowMessages.cs
--- a/Engine/Source/Programs/Shared/EpicGames.Horde/Issues/WorkflowMessages.cs
+++ b/Engine/Source/Programs/Shared/EpicGames.Horde/Issues/WorkflowMessages.cs
@@ -70,6 +70,9 @@ namespace EpicGames.Horde.Issues
 
 		IReadOnlyList<int> IWorkflow.EscalateRootCauseAnalysisTimes => EscalateRootCauseAnalysisTimes;
 
+		/// <inheritdoc/>
+		public bool AutoRcaEnabled { get; set; }
+
 		/// <inheritdoc/>
 		public int MaxMentions { get; set; }
 
@@ -129,6 +132,7 @@ namespace EpicGames.Horde.Issues
 			EscalateAlias = workflow.EscalateAlias;
 			EscalateTimes = workflow.EscalateTimes.ToList();
 			EscalateRootCauseAnalysisTimes = workflow.EscalateRootCauseAnalysisTimes.ToList();
+			AutoRcaEnabled = workflow.AutoRcaEnabled;
 			MaxMentions = workflow.MaxMentions;
 			AllowMentions = workflow.AllowMentions;
 			InviteRestrictedUsers = workflow.InviteRestrictedUsers;
