Horde: Compiler warning fixes.
#rnx

diff --git a/Engine/Source/Programs/Horde/Plugins/Build/HordeServer.Build/Issues/IssueService.cs b/Engine/Source/Programs/Horde/Plugins/Build/HordeServer.Build/Issues/IssueService.cs
--- a/Engine/Source/Programs/Horde/Plugins/Build/HordeServer.Build/Issues/IssueService.cs
+++ b/Engine/Source/Programs/Horde/Plugins/Build/HordeServer.Build/Issues/IssueService.cs
@@ -1853,7 +1853,7 @@ namespace HordeServer.Issues
 						{
 							_logger.LogError(ex, "RCA processing failed for issue {IssueId}", issue.Id);
 						}
-					});
+					}, CancellationToken.None);
 				}
 			}
 			else
