Horde: Output logs from child process memory tracking in the batch/lease logger

Before this change, these were sent to the agent's log which is not visible from log viewers in the Horde Dashboard.

diff --git a/Engine/Source/Programs/Horde/HordeAgent.Tests/Utility/WindowsProcessMemoryTrackerTests.cs b/Engine/Source/Programs/Horde/HordeAgent.Tests/Utility/WindowsProcessMemoryTrackerTests.cs
--- a/Engine/Source/Programs/Horde/HordeAgent.Tests/Utility/WindowsProcessMemoryTrackerTests.cs
+++ b/Engine/Source/Programs/Horde/HordeAgent.Tests/Utility/WindowsProcessMemoryTrackerTests.cs
@@ -42,7 +42,7 @@ public class WindowsProcessMemoryTrackerTests
 		}
 
 		using ManagedProcessGroup processGroup = new();
-		await using WindowsProcessMemoryTracker tracker = new(processGroup, TimeSpan.FromMilliseconds(100), _logger);
+		using WindowsProcessMemoryTracker tracker = new(processGroup, TimeSpan.FromMilliseconds(100), _logger);
 
 		tracker.Start();
 		await Task.Delay(150); // Let it run for a bit
@@ -62,7 +62,7 @@ public class WindowsProcessMemoryTrackerTests
 
 		using ManagedProcessGroup processGroup = new();
 
-		await using WindowsProcessMemoryTracker tracker = new(processGroup, TimeSpan.FromMilliseconds(100), _logger);
+		using WindowsProcessMemoryTracker tracker = new(processGroup, TimeSpan.FromMilliseconds(100), _logger);
 		tracker.Start();
 		await Task.Delay(50); // Give it time to sample
 
@@ -93,7 +93,7 @@ public class WindowsProcessMemoryTrackerTests
 			null,
 			ProcessPriorityClass.Normal);
 
-		await using WindowsProcessMemoryTracker tracker = new(
+		using WindowsProcessMemoryTracker tracker = new(
 			processGroup,
 			TimeSpan.FromMilliseconds(50),
 			_logger);
@@ -140,7 +140,7 @@ public class WindowsProcessMemoryTrackerTests
 			null,
 			ProcessPriorityClass.Normal);
 
-		await using WindowsProcessMemoryTracker tracker = new(
+		using WindowsProcessMemoryTracker tracker = new(
 			processGroup,
 			TimeSpan.FromMilliseconds(50),
 			_logger);
@@ -190,7 +190,7 @@ public class WindowsProcessMemoryTrackerTests
 			null,
 			ProcessPriorityClass.Normal);
 
-		await using WindowsProcessMemoryTracker tracker = new(
+		using WindowsProcessMemoryTracker tracker = new(
 			processGroup,
 			TimeSpan.FromMilliseconds(50),
 			_logger);
@@ -265,7 +265,7 @@ public class WindowsProcessMemoryTrackerTests
 			null,
 			ProcessPriorityClass.Normal);
 
-		await using WindowsProcessMemoryTracker tracker = new(
+		using WindowsProcessMemoryTracker tracker = new(
 			processGroup,
 			TimeSpan.FromMilliseconds(50),
 			_logger);
@@ -307,7 +307,7 @@ public class WindowsProcessMemoryTrackerTests
 
 		using ManagedProcessGroup processGroup = new();
 
-		await using WindowsProcessMemoryTracker tracker = new(
+		using WindowsProcessMemoryTracker tracker = new(
 			processGroup,
 			TimeSpan.FromMilliseconds(100),
 			_logger);
diff --git a/Engine/Source/Programs/Horde/HordeAgent/Leases/Handlers/JobHandler.cs b/Engine/Source/Programs/Horde/HordeAgent/Leases/Handlers/JobHandler.cs
--- a/Engine/Source/Programs/Horde/HordeAgent/Leases/Handlers/JobHandler.cs
+++ b/Engine/Source/Programs/Horde/HordeAgent/Leases/Handlers/JobHandler.cs
@@ -15,6 +15,7 @@ using HordeCommon.Rpc.Tasks;
 using Microsoft.Extensions.Logging;
 using Microsoft.Extensions.Options;
 using OpenTelemetry.Trace;
+using HordeAgent.Utility;
 
 namespace HordeAgent.Leases.Handlers
 {
@@ -27,6 +28,8 @@ namespace HordeAgent.Leases.Handlers
 		private readonly string? _driverCancelEventName;
 		private readonly EventWaitHandle? _driverCancelEvent;
 		private const string DriverCancelEventKeyName = "UE_HORDE_JOB_EXECUTOR_ABORT_EVENT_NAME";
+		const double BytesPerMegabyte = 1024.0 * 1024.0;
+		private WindowsProcessMemoryTracker? _memoryTracker;
 
 		/// <summary>
 		/// The cancellation source to make sure the job driver is always terminated 
@@ -87,6 +90,12 @@ namespace HordeAgent.Leases.Handlers
 				_driverProcessGroup?.Dispose();
 				_driverCancelEvent?.Dispose();
 				_cancellationSource.Dispose();
+				
+				// Normally disposed in ExecuteAsync, but handle edge cases
+				if (OperatingSystem.IsWindows())
+				{
+					_memoryTracker?.Dispose();
+				}
 			}
 			base.Dispose(disposing);
 		}
@@ -101,6 +110,13 @@ namespace HordeAgent.Leases.Handlers
 
 			await using IServerLogger logger = session.HordeClient.CreateServerLogger(LogId.Parse(executeTask.LogId)).WithLocalLogger(localLogger);
 
+			// Track Horde agent's own memory at job start (Windows only)
+			long agentMemoryStart = 0;
+			if (OperatingSystem.IsWindows())
+			{
+				agentMemoryStart = Process.GetCurrentProcess().PrivateMemorySize64;
+			}
+
 			int exitCode = 0;
 			Exception? exception = null;
 
@@ -139,6 +155,14 @@ namespace HordeAgent.Leases.Handlers
 
 				// Need to know the driver process ID for cancellation, use an alternative to RunDotNetProcessAsync
 				StartDriverProcess(driverAssembly, arguments, environment, AgentApp.IsSelfContained, logger);
+
+				// Start memory tracking for child processes (Windows only)
+				if (OperatingSystem.IsWindows() && _driverProcessGroup != null)
+				{
+					_memoryTracker = new WindowsProcessMemoryTracker(_driverProcessGroup, TimeSpan.FromSeconds(30), logger);
+					_memoryTracker.Start();
+				}
+
 				await using (cancellationToken.Register(DriverCancelHandler))
 				{
 					exitCode = await WaitForDriverProcessExitAsync(logger, _cancellationSource.Token);
@@ -162,6 +186,18 @@ namespace HordeAgent.Leases.Handlers
 				_driverCancelEvent?.Reset();
 			}
 
+			// Log memory metrics to server (Windows only)
+			if (OperatingSystem.IsWindows())
+			{
+				LogAgentMemory(agentMemoryStart, logger);
+				if (_driverProcessGroup != null && _memoryTracker != null)
+				{
+					LogChildProcessMemory(_driverProcessGroup, _memoryTracker, logger);
+				}
+				_memoryTracker?.Dispose();
+				_memoryTracker = null;
+			}
+
 			if (exception != null || exitCode != 0)
 			{
 				try
@@ -299,6 +335,52 @@ namespace HordeAgent.Leases.Handlers
 			int result = kill(processId, SigTerm);
 			return result == 0;
 		}
+
+		private static void LogAgentMemory(long memoryStart, ILogger logger)
+		{
+			using Process currentProcess = Process.GetCurrentProcess();
+			long memoryEnd = currentProcess.PrivateMemorySize64;
+			long peakMemory = currentProcess.PeakWorkingSet64;
+
+			logger.LogInformation(
+				"HordeAgent Memory: Start={StartMb:F1} MB End={EndMb:F1} MB Delta={DeltaMb:+0.0;-0.0;0.0} MB Peak={PeakMb:F1} MB",
+				memoryStart / BytesPerMegabyte,
+				memoryEnd / BytesPerMegabyte,
+				(memoryEnd - memoryStart) / BytesPerMegabyte,
+				peakMemory / BytesPerMegabyte);
+		}
+
+		[System.Runtime.Versioning.SupportedOSPlatform("windows")]
+		private static void LogChildProcessMemory(
+			ManagedProcessGroup processGroup,
+			WindowsProcessMemoryTracker tracker,
+			ILogger logger)
+		{
+			IReadOnlyList<ProcessMemoryInfo> topProcesses = tracker.GetTopMemoryProcesses(5);
+
+			long totalPeakBytes = processGroup.PeakJobMemoryUsed;
+			long largestProcessPeakBytes = processGroup.PeakProcessMemoryUsed;
+
+			logger.LogInformation(
+				"Child Process Memory Summary: TotalPeak={TotalPeakMb:F1} MB, LargestSingleProcess={LargestMb:F1} MB",
+				totalPeakBytes / BytesPerMegabyte,
+				largestProcessPeakBytes / BytesPerMegabyte);
+
+			if (topProcesses.Count > 0)
+			{
+				logger.LogInformation("Top {Count} memory-consuming processes:", topProcesses.Count);
+				foreach (ProcessMemoryInfo proc in topProcesses)
+				{
+					logger.LogInformation(
+						"  PID={Pid} {Name}.exe PeakPrivate={PrivateMb:F1} MB PeakWorkingSet={WorkingSetMb:F1} MB Args=\"{Args}\"",
+						proc.ProcessId,
+						proc.ProcessName,
+						proc.PeakPrivateBytes / BytesPerMegabyte,
+						proc.PeakWorkingSetBytes / BytesPerMegabyte,
+						proc.CommandLine);
+				}
+			}
+		}
 	}
 
 	class JobHandlerFactory : LeaseHandlerFactory<ExecuteJobTask>
diff --git a/Engine/Source/Programs/Horde/HordeAgent/Leases/LeaseHandler.cs b/Engine/Source/Programs/Horde/HordeAgent/Leases/LeaseHandler.cs
--- a/Engine/Source/Programs/Horde/HordeAgent/Leases/LeaseHandler.cs
+++ b/Engine/Source/Programs/Horde/HordeAgent/Leases/LeaseHandler.cs
@@ -7,7 +7,6 @@ using EpicGames.Horde.Agents.Metrics;
 using Google.Protobuf;
 using Google.Protobuf.WellKnownTypes;
 using HordeAgent.Services;
-using HordeAgent.Utility;
 using HordeCommon.Rpc;
 using HordeCommon.Rpc.Messages;
 using Microsoft.Extensions.Logging;
@@ -21,11 +20,6 @@ namespace HordeAgent.Leases
 	/// </summary>
 	abstract class LeaseHandler : IDisposable
 	{
-		/// <summary>
-		/// Conversion factor from bytes to megabytes
-		/// </summary>
-		const double BytesPerMegabyte = 1024.0 * 1024.0;
-
 		/// <summary>
 		/// Identifier for this lease
 		/// </summary>
@@ -130,13 +124,6 @@ namespace HordeAgent.Leases
 			Stopwatch leaseDuration = Stopwatch.StartNew();
 			CpuTimeMetrics? cpuTimeStart = systemMetrics.GetCpuTime();
 
-			// Track Horde agent's own memory at lease start (Windows only)
-			long agentMemoryStart = 0;
-			if (OperatingSystem.IsWindows())
-			{
-				agentMemoryStart = Process.GetCurrentProcess().PrivateMemorySize64;
-			}
-
 			// Get the lease outcome
 			LeaseResult result = new (LeaseOutcome.Failed, LeaseOutcomeReason.UnexpectedError);
 			try
@@ -156,12 +143,6 @@ namespace HordeAgent.Leases
 			}
 			CpuTimeMetrics? cpuTimeEnd = systemMetrics.GetCpuTime();
 
-			// Log HordeAgent's own memory delta (Windows only)
-			if (OperatingSystem.IsWindows())
-			{
-				LogAgentMemory(agentMemoryStart, logger);
-			}
-
 			// Update the state of the lease
 			RpcLease newRpcLease = new RpcLease(RpcLease);
 			if (_cancellationSource.IsCancellationRequested)
@@ -223,16 +204,7 @@ namespace HordeAgent.Leases
 			{
 				using ManagedProcessGroup processGroup = new ManagedProcessGroup();
 				using ManagedProcess process = new ManagedProcess(processGroup, executable, commandLine, null, environment, ProcessPriorityClass.Normal);
-
-				// Run the process with optional memory tracking
-				if (OperatingSystem.IsWindows())
-				{
-					return await RunProcessWithMemoryTrackingAsync(processGroup, process, logger, cancellationToken);
-				}
-				else
-				{
-					return await RunProcessCoreAsync(process, logger, cancellationToken);
-				}
+				return await RunProcessCoreAsync(process, logger, cancellationToken);
 			}
 			catch (Exception ex)
 			{
@@ -241,25 +213,6 @@ namespace HordeAgent.Leases
 			}
 		}
 
-		/// <summary>
-		/// Runs the process with memory tracking (Windows only)
-		/// </summary>
-		[System.Runtime.Versioning.SupportedOSPlatform("windows")]
-		private static async Task<int> RunProcessWithMemoryTrackingAsync(ManagedProcessGroup processGroup, ManagedProcess process, ILogger logger, CancellationToken cancellationToken)
-		{
-			await using WindowsProcessMemoryTracker memoryTracker = new (processGroup, TimeSpan.FromSeconds(30), logger);
-			memoryTracker.Start();
-
-			try
-			{
-				return await RunProcessCoreAsync(process, logger, cancellationToken);
-			}
-			finally
-			{
-				LogChildProcessMemory(processGroup, memoryTracker, logger);
-			}
-		}
-
 		/// <summary>
 		/// Core process execution loop
 		/// </summary>
@@ -350,59 +303,6 @@ namespace HordeAgent.Leases
 			}
 			return metrics;
 		}
-
-		/// <summary>
-		/// Logs HordeAgent's own memory usage at lease end
-		/// </summary>
-		private static void LogAgentMemory(long memoryStart, ILogger logger)
-		{
-			using Process currentProcess = Process.GetCurrentProcess();
-			long memoryEnd = currentProcess.PrivateMemorySize64;
-			long peakMemory = currentProcess.PeakWorkingSet64;
-
-			logger.LogInformation(
-				"HordeAgent Memory: Start={StartMb:F1} MB End={EndMb:F1} MB Delta={DeltaMb:+0.0;-0.0;0.0} MB Peak={PeakMb:F1} MB",
-				memoryStart / BytesPerMegabyte,
-				memoryEnd / BytesPerMegabyte,
-				(memoryEnd - memoryStart) / BytesPerMegabyte,
-				peakMemory / BytesPerMegabyte);
-		}
-
-		/// <summary>
-		/// Logs child process memory summary
-		/// </summary>
-		[System.Runtime.Versioning.SupportedOSPlatform("windows")]
-		private static void LogChildProcessMemory(
-			ManagedProcessGroup processGroup,
-			WindowsProcessMemoryTracker tracker,
-			ILogger logger)
-		{
-			IReadOnlyList<ProcessMemoryInfo> topProcesses = tracker.GetTopMemoryProcesses(5);
-
-			// Log job-level aggregates from ManagedProcessGroup
-			long totalPeakBytes = processGroup.PeakJobMemoryUsed;
-			long largestProcessPeakBytes = processGroup.PeakProcessMemoryUsed;
-
-			logger.LogInformation(
-				"Child Process Memory Summary: TotalPeak={TotalPeakMb:F1} MB, LargestSingleProcess={LargestMb:F1} MB",
-				totalPeakBytes / BytesPerMegabyte,
-				largestProcessPeakBytes / BytesPerMegabyte);
-
-			if (topProcesses.Count > 0)
-			{
-				logger.LogInformation("Top {Count} memory-consuming processes:", topProcesses.Count);
-				foreach (ProcessMemoryInfo proc in topProcesses)
-				{
-					logger.LogInformation(
-						"  PID={Pid} {Name}.exe PeakPrivate={PrivateMb:F1} MB PeakWorkingSet={WorkingSetMb:F1} MB Args=\"{Args}\"",
-						proc.ProcessId,
-						proc.ProcessName,
-						proc.PeakPrivateBytes / BytesPerMegabyte,
-						proc.PeakWorkingSetBytes / BytesPerMegabyte,
-						proc.CommandLine);
-				}
-			}
-		}
 	}
 
 	/// <summary>
diff --git a/Engine/Source/Programs/Horde/HordeAgent/Utility/WindowsProcessMemoryTracker.cs b/Engine/Source/Programs/Horde/HordeAgent/Utility/WindowsProcessMemoryTracker.cs
--- a/Engine/Source/Programs/Horde/HordeAgent/Utility/WindowsProcessMemoryTracker.cs
+++ b/Engine/Source/Programs/Horde/HordeAgent/Utility/WindowsProcessMemoryTracker.cs
@@ -32,7 +32,7 @@ namespace HordeAgent.Utility
 	/// Samples at configurable intervals and maintains peak memory values per process.
 	/// </summary>
 	[SupportedOSPlatform("windows")]
-	public sealed class WindowsProcessMemoryTracker : IAsyncDisposable
+	public sealed class WindowsProcessMemoryTracker : IDisposable
 	{
 		/// <summary>
 		/// Maximum length for stored command line strings
@@ -119,17 +119,19 @@ namespace HordeAgent.Utility
 		}
 
 		/// <inheritdoc/>
-		public async ValueTask DisposeAsync()
+		public void Dispose()
 		{
 			if (_cts != null)
 			{
-				await _cts.CancelAsync();
+				_cts.Cancel();
 
 				if (_samplingTask != null)
 				{
 					try
 					{
-						await _samplingTask;
+#pragma warning disable VSTHRD002 // Synchronous wait is unavoidable in Dispose
+						_samplingTask.GetAwaiter().GetResult();
+#pragma warning restore VSTHRD002
 					}
 					catch (OperationCanceledException)
 					{
