Horde: Fix compiler warning in WindowsProcessMemoryTrackerTests.cs

diff --git a/Engine/Source/Programs/Horde/HordeAgent.Tests/Utility/WindowsProcessMemoryTrackerTests.cs b/Engine/Source/Programs/Horde/HordeAgent.Tests/Utility/WindowsProcessMemoryTrackerTests.cs
--- a/Engine/Source/Programs/Horde/HordeAgent.Tests/Utility/WindowsProcessMemoryTrackerTests.cs
+++ b/Engine/Source/Programs/Horde/HordeAgent.Tests/Utility/WindowsProcessMemoryTrackerTests.cs
@@ -297,7 +297,7 @@ public class WindowsProcessMemoryTrackerTests
 
 	[TestMethod]
 	[TestCategory("Windows")]
-	public async Task TrackerThrowsWhenStartedTwiceAsync()
+	public void TrackerThrowsWhenStartedTwice()
 	{
 		if (!OperatingSystem.IsWindows())
 		{
