Horde: Fix race condition in agent relay long-poll mechanism

Capture notification task before reading state to prevent race condition where updates occurring between state read and task capture would be missed, causing long-poll timeouts

diff --git a/Engine/Source/Programs/Horde/Plugins/Compute/HordeServer.Compute/Agents/Relay/AgentRelayService.cs b/Engine/Source/Programs/Horde/Plugins/Compute/HordeServer.Compute/Agents/Relay/AgentRelayService.cs
--- a/Engine/Source/Programs/Horde/Plugins/Compute/HordeServer.Compute/Agents/Relay/AgentRelayService.cs
+++ b/Engine/Source/Programs/Horde/Plugins/Compute/HordeServer.Compute/Agents/Relay/AgentRelayService.cs
@@ -93,6 +93,7 @@ public sealed class AgentRelayService : RelayRpc.RelayRpcBase, IHostedService, I
 	private int _minPort = 10000;
 	private int _maxPort = 30000; // Linux ephemeral port range is usually 32768-60999 so staying below 30000 seems like a sane default
 	private Random _random = new ();
+	private Func<Task>? _testDelayBeforeTaskCapture;
 
 	/// <summary>
 	/// Constructor
@@ -467,7 +468,12 @@ public sealed class AgentRelayService : RelayRpc.RelayRpcBase, IHostedService, I
 	{
 		_random = new Random(seed);
 	}
-	
+
+	internal void SetTestDelayBeforeTaskCapture(Func<Task>? delay)
+	{
+		_testDelayBeforeTaskCapture = delay;
+	}
+
 	private static void AddPortMappingMetadataToSpan(TelemetrySpan span, PortMapping portMapping)
 	{
 		span.SetAttribute("leaseId", portMapping.LeaseId);
@@ -502,13 +508,20 @@ public sealed class AgentRelayService : RelayRpc.RelayRpcBase, IHostedService, I
 			ClusterId clusterId = new(request.ClusterId);
 			await UpdateAgentHeartbeatAsync(clusterId, request.AgentId, request.IpAddresses.Select(IPAddress.Parse));
 
+			// Capture notification task BEFORE reading state to avoid race condition
+			// By capturing first, we always get a task that will be signaled if an update occurs.
+			Task mappingUpdatedTask = _onPortMappingUpdated.Task;
+			
 			(int revisionCount, List<PortMapping> portMappings) = await GetPortMappingsAsync(clusterId);
+			if (_testDelayBeforeTaskCapture != null)
+			{
+				await _testDelayBeforeTaskCapture();
+			}
 			if (request.RevisionCount == revisionCount || request.RevisionCount == -1)
 			{
 				// Client is in sync with latest revision server has or explicitly requested long-polling by setting revision to -1
 				// Long-poll for any new updates...
 
-				Task mappingUpdatedTask = _onPortMappingUpdated.Task;
 				Task result = await Task.WhenAny(mappingUpdatedTask, Task.Delay(_longPollTimeout, context.CancellationToken));
 				if (result == mappingUpdatedTask)
 				{
