Horde: Increase max command line length for process memory tracking

diff --git a/Engine/Source/Programs/Horde/HordeAgent/Utility/WindowsProcessMemoryTracker.cs b/Engine/Source/Programs/Horde/HordeAgent/Utility/WindowsProcessMemoryTracker.cs
--- a/Engine/Source/Programs/Horde/HordeAgent/Utility/WindowsProcessMemoryTracker.cs
+++ b/Engine/Source/Programs/Horde/HordeAgent/Utility/WindowsProcessMemoryTracker.cs
@@ -37,7 +37,7 @@ namespace HordeAgent.Utility
 		/// <summary>
 		/// Maximum length for stored command line strings
 		/// </summary>
-		public const int MaxCommandLineLength = 100;
+		public const int MaxCommandLineLength = 500;
 
 		#region Native Interop
 
