Horde: Add error handling for stuck step timeout check in JobExpirationService

diff --git a/Engine/Source/Programs/Horde/Plugins/Build/HordeServer.Build/Jobs/JobExpirationService.cs b/Engine/Source/Programs/Horde/Plugins/Build/HordeServer.Build/Jobs/JobExpirationService.cs
--- a/Engine/Source/Programs/Horde/Plugins/Build/HordeServer.Build/Jobs/JobExpirationService.cs
+++ b/Engine/Source/Programs/Horde/Plugins/Build/HordeServer.Build/Jobs/JobExpirationService.cs
@@ -87,10 +87,19 @@ namespace HordeServer.Jobs
 			DateTime maxStepStartTime = currentTime - StepTimeout;
 
 			// Find jobs with running batches that might have stuck steps
-			IReadOnlyList<IJob> jobs = await _jobCollection.FindAsync(
-				new FindJobOptions(BatchState: JobStepBatchState.Running),
-				count: 100,
-				cancellationToken: cancellationToken);
+			IReadOnlyList<IJob> jobs;
+			try
+			{
+				jobs = await _jobCollection.FindAsync(
+					new FindJobOptions(BatchState: JobStepBatchState.Running),
+					count: 100,
+					cancellationToken: cancellationToken);
+			}
+			catch (ArgumentOutOfRangeException ex)
+			{
+				_logger.LogWarning(ex, "Failed to find jobs with running batches, skipping stuck step timeout check");
+				return;
+			}
 
 			foreach (IJob job in jobs)
 			{
