Horde: Disable background ticker for cleaning up stuck job steps

Currently causing exceptions which may impact production deploy.

diff --git a/Engine/Source/Programs/Horde/Plugins/Build/HordeServer.Build/Jobs/JobExpirationService.cs b/Engine/Source/Programs/Horde/Plugins/Build/HordeServer.Build/Jobs/JobExpirationService.cs
--- a/Engine/Source/Programs/Horde/Plugins/Build/HordeServer.Build/Jobs/JobExpirationService.cs
+++ b/Engine/Source/Programs/Horde/Plugins/Build/HordeServer.Build/Jobs/JobExpirationService.cs
@@ -51,7 +51,7 @@ namespace HordeServer.Jobs
 			DateTime currentTime = _clock.UtcNow;
 
 			// Timeout stuck running steps
-			await TimeoutStuckStepsAsync(currentTime, cancellationToken);
+			// await TimeoutStuckStepsAsync(currentTime, cancellationToken);
 
 			// Expire old jobs
 			BuildConfig buildConfig = _buildConfig.CurrentValue;
