Horde: Disable stuck job steps test

Due to CL 50304178. Also fix warning from disabling this in JobExpirationService.

diff --git a/Engine/Source/Programs/Horde/Plugins/Build/HordeServer.Build.Tests/Jobs/JobExpirationTests.cs b/Engine/Source/Programs/Horde/Plugins/Build/HordeServer.Build.Tests/Jobs/JobExpirationTests.cs
--- a/Engine/Source/Programs/Horde/Plugins/Build/HordeServer.Build.Tests/Jobs/JobExpirationTests.cs
+++ b/Engine/Source/Programs/Horde/Plugins/Build/HordeServer.Build.Tests/Jobs/JobExpirationTests.cs
@@ -76,6 +76,7 @@ namespace HordeServer.Tests.Jobs
 		}
 
 		[TestMethod]
+		[Ignore] // See CL 50304178
 		public async Task StuckRunningStepShouldBeTimedOutAsync()
 		{
 			await ServiceProvider.GetRequiredService<JobExpirationService>().StartAsync(default);
diff --git a/Engine/Source/Programs/Horde/Plugins/Build/HordeServer.Build/Jobs/JobExpirationService.cs b/Engine/Source/Programs/Horde/Plugins/Build/HordeServer.Build/Jobs/JobExpirationService.cs
--- a/Engine/Source/Programs/Horde/Plugins/Build/HordeServer.Build/Jobs/JobExpirationService.cs
+++ b/Engine/Source/Programs/Horde/Plugins/Build/HordeServer.Build/Jobs/JobExpirationService.cs
@@ -1,5 +1,6 @@
 ﻿// Copyright Epic Games, Inc. All Rights Reserved.
 
+using System.Diagnostics.CodeAnalysis;
 using EpicGames.Core;
 using EpicGames.Horde.Jobs;
 using HordeServer.Streams;
@@ -82,6 +83,7 @@ namespace HordeServer.Jobs
 			}
 		}
 
+		[SuppressMessage("CodeQuality", "IDE0051:Remove unused private members", Justification = "Stuck job steps disabled")]
 		async Task TimeoutStuckStepsAsync(DateTime currentTime, CancellationToken cancellationToken)
 		{
 			DateTime maxStepStartTime = currentTime - StepTimeout;
