Horde: Missed file for agent telemetry

diff --git a/Engine/Source/Programs/Horde/HordeDashboard/src/backend/Backend.ts b/Engine/Source/Programs/Horde/HordeDashboard/src/backend/Backend.ts
--- a/Engine/Source/Programs/Horde/HordeDashboard/src/backend/Backend.ts
+++ b/Engine/Source/Programs/Horde/HordeDashboard/src/backend/Backend.ts
@@ -266,7 +266,7 @@ export class Backend {
     getAgentLeases(agentId: string, minTime: Date, maxTime: Date): Promise<GetAgentLeaseResponse[]> {
         const params: string[] = [];
         params.push("startTime=" + minTime.toISOString());
-        params.push("finishime=" + maxTime.toISOString());
+        params.push("finishTime=" + maxTime.toISOString());
         const paramString = params.join("&");
         return new Promise<GetAgentLeaseResponse[]>((resolve, reject) => {
             this._fetch.get(`/api/v1/agents/${agentId}/leases?${paramString}`).then((response) => {
