Horde: Update release notes

#rnx

diff --git a/Engine/Source/Programs/Horde/Docs/ReleaseNotes.md b/Engine/Source/Programs/Horde/Docs/ReleaseNotes.md
--- a/Engine/Source/Programs/Horde/Docs/ReleaseNotes.md
+++ b/Engine/Source/Programs/Horde/Docs/ReleaseNotes.md
@@ -1,5 +1,32 @@
 # Release Notes
 
+## 2026-01-29
+
+* Add error handling for stuck step timeout check in JobExpirationService (50303167)
+* Increase max command line length for process memory tracking (50300448)
+* Fix race condition in agent relay long-poll mechanism (50299104)
+* Fix compiler warning in WindowsProcessMemoryTrackerTests.cs (50296706)
+* Output logs from child process memory tracking in the batch/lease logger (50296181)
+* Add IRcaProcessor framework and config options. 
+  Enables licensees to implement their own root cause analysis pipelines by introducing an extensible IRcaProcessor interface in HordeServer.Shared. (50284420)
+* Add bidirectional commit queries, parallel scoring, and fingerprint fallback for issue processing
+  Support ascending/descending sort order in ICommitCollection.FindAsync. Enable optional bidirectional suspect search that
+  queries from both ends of the commit range to better identify root causes. Add parallel commit scoring with configurable
+  concurrency. Add fingerprint-based fallback to match existing issues when no suspects found, preventing duplicate issues
+  when commit range exceeds MaxChanges. All features are config-driven and default to off. (50256512)
+* Remove vertical boundary lines from telemetry lease overlays and hide overlays entirely when all leases are too narrow (<4px) to be visually meaningful. (50250137)
+* Add ACL-gated AI Analysis link to issue modal (50186747)
+* Implement dashboard plugin permission check (50177903)
+* Fix stuck job/lease when retrying or aborting running steps
+  When retrying or aborting a running step, mark the step for abort and set the batch error to NoLongerNeeded or Cancelled. This signals that the lease should be cancelled. Also add a 24-hour timeout for stuck running steps that never complete, with periodic detection via JobExpirationService. (50164696)
+* Adding stream aliasing feature backend and frontend (50161443)
+* Fix sparklines cloning lease overlays (50160749)
+* Updated perf record for use in new API endpoints with more context. (50148423)
+* Update placeholder to handle between "New" and "Select..." labels based on admin access. (50124878)
+* Add a tool tip for cancelled batch texts explaining the reasons (50099723)
+* Propagate lease cancellation reason for cancelled job batches (50099511)
+* Support both ordered and unordered commit id. (50084686)
+
 ## 2026-01-22
 
 * Reducing the severity of noisy error on duplicate payload properties due to the benign nature. (50024919)
