HordeAgent: Ensure we use the targeted framework (.net 8)

By default we use LatestMajor which may not be desired.

#Synced-CL 49593614

diff --git a/Engine/Source/Programs/Horde/HordeAgent/HordeAgent.csproj b/Engine/Source/Programs/Horde/HordeAgent/HordeAgent.csproj
--- a/Engine/Source/Programs/Horde/HordeAgent/HordeAgent.csproj
+++ b/Engine/Source/Programs/Horde/HordeAgent/HordeAgent.csproj
@@ -4,6 +4,7 @@
   <PropertyGroup>
     <OutputType>Exe</OutputType>
     <TargetFramework>net8.0</TargetFramework>
+	<RollForward>LatestMinor</RollForward>
     <Nullable>enable</Nullable>
     <ImplicitUsings>enable</ImplicitUsings>
 	<ApplicationIcon>HordeAgent.ico</ApplicationIcon>
