Horde: Make upgrade of Windows agent more robust with retry logic

Add retry for file deletion and move operations during agent self-upgrade to handle lingering file locks after service stop.

diff --git a/Engine/Source/Programs/Horde/HordeAgent.Tests/Commands/ServiceUpgradeTests.cs b/Engine/Source/Programs/Horde/HordeAgent.Tests/Commands/ServiceUpgradeTests.cs
new file mode 100644
--- /dev/null
+++ b/Engine/Source/Programs/Horde/HordeAgent.Tests/Commands/ServiceUpgradeTests.cs
@@ -0,0 +1,118 @@
+// Copyright Epic Games, Inc. All Rights Reserved.
+
+using System;
+using System.IO;
+using System.Threading;
+using System.Threading.Tasks;
+using HordeAgent.Commands.Service;
+using Microsoft.Extensions.Logging.Abstractions;
+using Microsoft.VisualStudio.TestTools.UnitTesting;
+
+namespace HordeAgent.Tests.Commands;
+
+[TestClass]
+public sealed class ServiceUpgradeTests : IDisposable
+{
+	private readonly string _testDir = Path.Combine(Path.GetTempPath(), $"upgrade-test-{Guid.NewGuid()}");
+
+	public ServiceUpgradeTests() => Directory.CreateDirectory(_testDir);
+
+	public void Dispose() => Try(() => Directory.Delete(_testDir, true));
+
+	[TestMethod]
+	[Timeout(30000)]
+	public async Task DeleteFileWithRetry_LockedFile_RetriesAndSucceedsAsync()
+	{
+		string testFile = Path.Combine(_testDir, "locked.dll");
+		await File.WriteAllTextAsync(testFile, "test content");
+
+		await WithLockedFileAsync(testFile, 1500, () =>
+			UpgradeCommand.DeleteFileWithRetry(NullLogger.Instance, testFile, [500, 1000, 2000]));
+
+		Assert.IsFalse(File.Exists(testFile));
+	}
+
+	[TestMethod]
+	public void DeleteFileWithRetry_FileAlreadyDeleted_Succeeds()
+	{
+		string testFile = Path.Combine(_testDir, "nonexistent.dll");
+		UpgradeCommand.DeleteFileWithRetry(NullLogger.Instance, testFile, [100]);
+	}
+
+	[TestMethod]
+	public void MoveFileWithRetry_AlreadyMoved_Succeeds()
+	{
+		string targetFile = Path.Combine(_testDir, "source.dll");
+		File.WriteAllText(targetFile, "already moved content");
+
+		UpgradeCommand.MoveFileWithRetry(NullLogger.Instance, targetFile + ".new", targetFile, [100]);
+
+		Assert.AreEqual("already moved content", File.ReadAllText(targetFile));
+	}
+
+	[TestMethod]
+	[Timeout(30000)]
+	public async Task MoveFileWithRetry_LockedTarget_RetriesAndSucceedsAsync()
+	{
+		string sourceFile = Path.Combine(_testDir, "new.dll.new");
+		string targetFile = Path.Combine(_testDir, "new.dll");
+		await File.WriteAllTextAsync(sourceFile, "new content");
+		await File.WriteAllTextAsync(targetFile, "old content");
+
+		await WithLockedFileAsync(targetFile, 1500, () =>
+			UpgradeCommand.MoveFileWithRetry(NullLogger.Instance, sourceFile, targetFile, [500, 1000, 2000]));
+
+		Assert.IsFalse(File.Exists(sourceFile));
+		Assert.AreEqual("new content", await File.ReadAllTextAsync(targetFile));
+	}
+
+	[TestMethod]
+	public void MoveFileWithRetry_NormalMove_Succeeds()
+	{
+		string sourceFile = Path.Combine(_testDir, "normal.dll.new");
+		string targetFile = Path.Combine(_testDir, "normal.dll");
+		File.WriteAllText(sourceFile, "content to move");
+
+		UpgradeCommand.MoveFileWithRetry(NullLogger.Instance, sourceFile, targetFile, [100]);
+
+		Assert.IsFalse(File.Exists(sourceFile));
+		Assert.AreEqual("content to move", File.ReadAllText(targetFile));
+	}
+
+	private static async Task WithLockedFileAsync(string path, int lockDurationMs, Action action)
+	{
+		using ManualResetEventSlim lockAcquired = new();
+		Task lockTask = Task.Run(async () =>
+		{
+			await using FileStream fs = File.Open(path, FileMode.Open, FileAccess.Read, FileShare.None);
+			lockAcquired.Set();
+			await Task.Delay(lockDurationMs);
+		});
+
+		lockAcquired.Wait();
+		action();
+		await TryAsync(() => lockTask);
+	}
+
+	private static void Try(Action action)
+	{
+		try
+		{
+			action();
+		}
+		catch
+		{
+		}
+	}
+
+	private static async Task TryAsync(Func<Task> action)
+	{
+		try
+		{
+			await action();
+		}
+		catch
+		{
+		}
+	}
+}
diff --git a/Engine/Source/Programs/Horde/HordeAgent/Commands/Service/ServiceUpgrade.cs b/Engine/Source/Programs/Horde/HordeAgent/Commands/Service/ServiceUpgrade.cs
--- a/Engine/Source/Programs/Horde/HordeAgent/Commands/Service/ServiceUpgrade.cs
+++ b/Engine/Source/Programs/Horde/HordeAgent/Commands/Service/ServiceUpgrade.cs
@@ -107,21 +107,67 @@ namespace HordeAgent.Commands.Service
 			return Task.FromResult(0);
 		}
 
-		static void UpgradeFilesInPlace(ILogger logger, HashSet<string> targetFiles, List<Tuple<string, string>> renameFiles)
+		private static readonly int[] s_retryDelaysMs = { 500, 1000, 2000, 3000, 5000 };
+
+		private static void UpgradeFilesInPlace(ILogger logger, HashSet<string> targetFiles, List<Tuple<string, string>> moveFiles)
 		{
 			// Remove all the target files
 			foreach (string targetFile in targetFiles)
 			{
-				logger.LogDebug("Deleting {File}", targetFile);
-				File.SetAttributes(targetFile, FileAttributes.Normal);
-				File.Delete(targetFile);
+				DeleteFileWithRetry(logger, targetFile, s_retryDelaysMs);
+			}
+
+			// Move all the new files into place
+			foreach (Tuple<string, string> pair in moveFiles)
+			{
+				MoveFileWithRetry(logger, pair.Item1, pair.Item2, s_retryDelaysMs);
+			}
+		}
+
+		internal static void DeleteFileWithRetry(ILogger logger, string filePath, int[] retryDelaysMs)
+		{
+			for (int attempt = 0; attempt <= retryDelaysMs.Length; attempt++)
+			{
+				try
+				{
+					if (!File.Exists(filePath))
+					{
+						return;
+					}
+					logger.LogDebug("Deleting {File} (attempt {Attempt})", filePath, attempt + 1);
+					FileUtils.ForceDeleteFile(filePath);
+					return;
+				}
+				catch (Exception ex) when (attempt < retryDelaysMs.Length)
+				{
+					logger.LogWarning(ex, "Failed to delete {File}, retrying in {Delay} ms... ({Attempt}/{Max})",
+						filePath, retryDelaysMs[attempt], attempt + 1, retryDelaysMs.Length + 1);
+					Thread.Sleep(retryDelaysMs[attempt]);
+				}
 			}
+		}
 
-			// Rename all the new files into place
-			foreach (Tuple<string, string> pair in renameFiles)
+		internal static void MoveFileWithRetry(ILogger logger, string sourceFile, string targetFile, int[] retryDelaysMs)
+		{
+			for (int attempt = 0; attempt <= retryDelaysMs.Length; attempt++)
 			{
-				logger.LogDebug("Renaming {SourceFile} to {TargetFile}", pair.Item1, pair.Item2);
-				File.Move(pair.Item1, pair.Item2, true);
+				try
+				{
+					// Idempotent check: if source is gone and target exists, move was already done
+					if (!File.Exists(sourceFile) && File.Exists(targetFile))
+					{
+						return;
+					}
+					logger.LogDebug("Moving {SourceFile} to {TargetFile} (attempt {Attempt})", sourceFile, targetFile, attempt + 1);
+					FileUtils.ForceMoveFile(sourceFile, targetFile);
+					return;
+				}
+				catch (Exception ex) when (attempt < retryDelaysMs.Length)
+				{
+					logger.LogWarning(ex, "Failed to move {SourceFile} to {TargetFile}, retrying in {Delay} ms... ({Attempt}/{Max})",
+						sourceFile, targetFile, retryDelaysMs[attempt], attempt + 1, retryDelaysMs.Length + 1);
+					Thread.Sleep(retryDelaysMs[attempt]);
+				}
 			}
 		}
 
@@ -157,9 +203,13 @@ namespace HordeAgent.Commands.Service
 			{
 				logger.LogInformation("Stopping service...");
 				service.Stop();
+				service.WaitForStatus(ServiceControllerStatus.Stopped, TimeSpan.FromSeconds(30));
 			}
 			otherProcess.WaitForExit();
 
+			// Additional buffer for Windows to release all file handles (antivirus, indexers, etc.)
+			Thread.Sleep(5000);
+
 			UpgradeFilesInPlace(logger, targetFiles, renameFiles);
 
 			// Run the new application
