[Horde Server][Horde Agent][Analytics] Propagation of Agent's system drive usage stats to Agent telemetry reporting

Modifies the original set of system metrics propagated into mongodb to split the concept of agent disk (what we're currently tracking) and system disk (C:\ on Windows, '/' on Unix. This data will eventually be consumed by the Horde telemetry dashboard to show a 2nd graph tracking system disk size alongside the pre-existing agent disk size graph.

#Synced-CL 48478172

diff --git a/Engine/Source/Programs/Horde/Drivers/JobDriver.Tests/Execution/FakeSystemMetrics.cs b/Engine/Source/Programs/Horde/Drivers/JobDriver.Tests/Execution/FakeSystemMetrics.cs
--- a/Engine/Source/Programs/Horde/Drivers/JobDriver.Tests/Execution/FakeSystemMetrics.cs
+++ b/Engine/Source/Programs/Horde/Drivers/JobDriver.Tests/Execution/FakeSystemMetrics.cs
@@ -19,7 +19,13 @@ internal class FakeSystemMetrics : ISystemMetrics
 	public MemoryMetrics? GetMemory() => null;
 
 	/// <inheritdoc/>
-	public DiskMetrics? GetDisk()
+	public DiskMetrics? GetWorkingDisk()
+	{
+		return new DiskMetrics { FreeSpace = FreeSpace, TotalSize = TotalSize };
+	}
+
+	/// <inheritdoc/>
+	public DiskMetrics? GetSystemDisk()
 	{
 		return new DiskMetrics { FreeSpace = FreeSpace, TotalSize = TotalSize };
 	}
diff --git a/Engine/Source/Programs/Horde/Drivers/JobDriver/Execution/ConformExecutor.cs b/Engine/Source/Programs/Horde/Drivers/JobDriver/Execution/ConformExecutor.cs
--- a/Engine/Source/Programs/Horde/Drivers/JobDriver/Execution/ConformExecutor.cs
+++ b/Engine/Source/Programs/Horde/Drivers/JobDriver/Execution/ConformExecutor.cs
@@ -160,7 +160,7 @@ namespace JobDriver.Execution
 
 		bool HasAvailableDiskFreeSpace(IList<RpcAgentWorkspace> pendingWorkspaces)
 		{
-			long? freeDiskSpace = _systemMetrics.GetDisk()?.FreeSpace;
+			long? freeDiskSpace = _systemMetrics.GetWorkingDisk()?.FreeSpace;
 			if (freeDiskSpace == null)
 			{
 				return false;
diff --git a/Engine/Source/Programs/Horde/Drivers/JobDriver/Execution/DefaultSystemMetrics.cs b/Engine/Source/Programs/Horde/Drivers/JobDriver/Execution/DefaultSystemMetrics.cs
--- a/Engine/Source/Programs/Horde/Drivers/JobDriver/Execution/DefaultSystemMetrics.cs
+++ b/Engine/Source/Programs/Horde/Drivers/JobDriver/Execution/DefaultSystemMetrics.cs
@@ -36,39 +36,39 @@ public class DefaultSystemMetrics : ISystemMetrics
 	public MemoryMetrics? GetMemory() => null;
 
 	/// <inheritdoc/>
-	public DiskMetrics? GetDisk()
+	public DiskMetrics? GetWorkingDisk() => GetDisk(EpicGames.Horde.Agents.Metrics.DriveType.Working);
+
+	/// <inheritdoc/>
+	public DiskMetrics? GetSystemDisk() => GetDisk(EpicGames.Horde.Agents.Metrics.DriveType.System);
+
+	/// <summary>
+	/// Fetches root directory metrics based off drive classification
+	/// </summary>
+	/// <param name="type">Classification of the drive we're fetching telemetry for</param>
+	/// <returns>DiskMetrics object representing total size and free space</returns>
+	private DiskMetrics? GetDisk(EpicGames.Horde.Agents.Metrics.DriveType type)
 	{
-		if (_workingDir == null)
+		string? rootDir;
+		if (RuntimeInformation.IsOSPlatform(OSPlatform.Windows))
 		{
-			_logger.LogError("Cannot get disk metrics: no working directory");
-			return null;
+			rootDir = Path.GetPathRoot(type == EpicGames.Horde.Agents.Metrics.DriveType.Working ? _workingDir?.FullName : Environment.SystemDirectory);
 		}
-
-		try
+		else
 		{
-			string? rootDir;
-			if (RuntimeInformation.IsOSPlatform(OSPlatform.Windows))
-			{
-				rootDir = Path.GetPathRoot(_workingDir.FullName);
-			}
-			else
-			{
-				rootDir = _workingDir.FullName;
-			}
-
-			if (rootDir == null)
-			{
-				_logger.LogError("Cannot get disk metrics: no root directory");
-				return null;
-			}
-
-			DriveInfo driveInfo = new(rootDir);
-			return new DiskMetrics { FreeSpace = driveInfo.AvailableFreeSpace, TotalSize = driveInfo.TotalSize };
+			rootDir = type == EpicGames.Horde.Agents.Metrics.DriveType.Working ? _workingDir?.FullName : "/";
 		}
-		catch (Exception error)
+
+		if (rootDir == null)
 		{
-			_logger.LogError(error, "Cannot get disk metrics: {Message}", error.Message);
+			_logger.LogError("Cannot get disk metrics: no root directory");
 			return null;
 		}
+
+		DriveInfo driveInfo = new(rootDir);
+		return new DiskMetrics
+		{
+			FreeSpace = driveInfo.AvailableFreeSpace,
+			TotalSize = driveInfo.TotalSize
+		};
 	}
 }
diff --git a/Engine/Source/Programs/Horde/HordeAgent/Leases/LeaseManager.cs b/Engine/Source/Programs/Horde/HordeAgent/Leases/LeaseManager.cs
--- a/Engine/Source/Programs/Horde/HordeAgent/Leases/LeaseManager.cs
+++ b/Engine/Source/Programs/Horde/HordeAgent/Leases/LeaseManager.cs
@@ -572,9 +572,11 @@ namespace HordeAgent.Leases
 				{
 					CpuMetrics? cpuMetrics = _systemMetrics.GetCpu();
 					MemoryMetrics? memMetrics = _systemMetrics.GetMemory();
-					DiskMetrics? diskMetrics = _systemMetrics.GetDisk();
 
-					if (cpuMetrics != null || memMetrics != null || diskMetrics != null)
+					DiskMetrics? workingDiskMetrics = _systemMetrics.GetWorkingDisk();
+					DiskMetrics? systemDiskMetrics = _systemMetrics.GetSystemDisk();
+
+					if (cpuMetrics != null || memMetrics != null || workingDiskMetrics != null || systemDiskMetrics != null)
 					{
 						RpcUploadTelemetryRequest request = new RpcUploadTelemetryRequest();
 						request.AgentId = _session.AgentId.ToString();
@@ -590,10 +592,15 @@ namespace HordeAgent.Leases
 							request.FreeRam = memMetrics.Available / 1024;
 							request.UsedRam = memMetrics.Used / 1024;
 						}
-						if (diskMetrics != null)
+						if (workingDiskMetrics != null)
+						{
+							request.WorkingFreeDisk = (ulong)(workingDiskMetrics.FreeSpace / (1024 * 1024));
+							request.WorkingTotalDisk = (ulong)(workingDiskMetrics.TotalSize / (1024 * 1024));
+						}
+						if (systemDiskMetrics != null)
 						{
-							request.FreeDisk = (ulong)(diskMetrics.FreeSpace / (1024 * 1024));
-							request.TotalDisk = (ulong)(diskMetrics.TotalSize / (1024 * 1024));
+							request.SystemFreeDisk = (ulong)(systemDiskMetrics.FreeSpace / (1024 * 1024));
+							request.SystemTotalDisk = (ulong)(systemDiskMetrics.TotalSize / (1024 * 1024));
 						}
 
 						await hordeRpc.UploadTelemetryAsync(request, cancellationToken: cancellationToken);
diff --git a/Engine/Source/Programs/Horde/HordeAgent/Services/TelemetryService.cs b/Engine/Source/Programs/Horde/HordeAgent/Services/TelemetryService.cs
--- a/Engine/Source/Programs/Horde/HordeAgent/Services/TelemetryService.cs
+++ b/Engine/Source/Programs/Horde/HordeAgent/Services/TelemetryService.cs
@@ -28,7 +28,10 @@ public sealed class DefaultSystemMetrics : ISystemMetrics
 	public MemoryMetrics? GetMemory() => null;
 
 	/// <inheritdoc/>
-	public DiskMetrics? GetDisk() => null;
+	public DiskMetrics? GetWorkingDisk() => null;
+
+	/// <inheritdoc/>
+	public DiskMetrics? GetSystemDisk() => null;
 }
 
 /// <summary>
@@ -39,7 +42,8 @@ public sealed class WindowsSystemMetrics : ISystemMetrics
 {
 	private readonly DirectoryReference _workingDir;
 	private readonly Stopwatch _driveInfoUpdateTimer = Stopwatch.StartNew();
-	private DriveInfo? _driveInfo;
+	private DriveInfo? _workingDriveInfo;
+	private DriveInfo? _systemDriveInfo;
 
 	/// <summary>
 	/// Constructor
@@ -47,7 +51,8 @@ public sealed class WindowsSystemMetrics : ISystemMetrics
 	public WindowsSystemMetrics(DirectoryReference workingDir)
 	{
 		_workingDir = workingDir;
-		_driveInfo = GetDriveInfo();
+		_workingDriveInfo = GetDriveInfo(EpicGames.Horde.Agents.Metrics.DriveType.Working);
+		_systemDriveInfo = GetDriveInfo(EpicGames.Horde.Agents.Metrics.DriveType.System);
 
 		GetCpu(); // Trigger this to ensure performance counter has a fetched value. Avoids an initial zero result when called later.
 	}
@@ -97,31 +102,64 @@ public sealed class WindowsSystemMetrics : ISystemMetrics
 	}
 
 	/// <inheritdoc/>
-	public DiskMetrics? GetDisk()
+	public DiskMetrics? GetWorkingDisk()
 	{
-		// Update the metrics if we don't have a recent sample
-		if (_driveInfoUpdateTimer.Elapsed > TimeSpan.FromSeconds(20))
+		RefreshDriveInfo();
+
+		if (_workingDriveInfo == null)
 		{
-			_driveInfo = GetDriveInfo();
-			_driveInfoUpdateTimer.Restart();
+			return null;
 		}
 
-		if (_driveInfo == null)
+		return new DiskMetrics
+		{
+			FreeSpace = _workingDriveInfo.AvailableFreeSpace,
+			TotalSize = _workingDriveInfo.TotalSize
+		};
+	}
+
+	/// <inheritdoc/>
+	public DiskMetrics? GetSystemDisk()
+	{
+		RefreshDriveInfo();
+
+		if (_systemDriveInfo == null)
 		{
 			return null;
 		}
 
-		return new DiskMetrics { FreeSpace = _driveInfo.AvailableFreeSpace, TotalSize = _driveInfo.TotalSize };
+		return new DiskMetrics
+		{
+			FreeSpace = _systemDriveInfo.AvailableFreeSpace,
+			TotalSize = _systemDriveInfo.TotalSize
+		};
 	}
 
-	DriveInfo? GetDriveInfo()
+	/// <summary>
+	/// Restart the drive info poll timer if the time period has elapsed
+	/// </summary>
+	void RefreshDriveInfo()
 	{
-		string? rootDir = Path.GetPathRoot(_workingDir.FullName);
-		if (rootDir == null)
+		if (_driveInfoUpdateTimer.Elapsed > TimeSpan.FromSeconds(20))
 		{
-			return null;
+			_workingDriveInfo = GetDriveInfo(EpicGames.Horde.Agents.Metrics.DriveType.Working);
+			_systemDriveInfo = GetDriveInfo(EpicGames.Horde.Agents.Metrics.DriveType.System);
+			_driveInfoUpdateTimer.Restart();
 		}
-		return new DriveInfo(rootDir);
+	}
+
+	/// <summary>
+	/// Fetches root directory based off drive classification
+	/// </summary>
+	/// <param name="type">Classification of the drive we're fetching telemetry for</param>
+	/// <returns>DiskMetrics object representing total size and free space</returns>
+	private DriveInfo? GetDriveInfo(EpicGames.Horde.Agents.Metrics.DriveType type)
+	{
+		string? rootDir = type == EpicGames.Horde.Agents.Metrics.DriveType.Working
+			? Path.GetPathRoot(_workingDir.FullName)
+			: Path.GetPathRoot(Environment.SystemDirectory);
+
+		return rootDir == null ? null : new DriveInfo(rootDir);
 	}
 }
 
@@ -252,14 +290,25 @@ public sealed class LinuxSystemMetrics : ISystemMetrics, IDisposable
 			return null;
 		}
 	}
-	
+
+	/// <inheritdoc/>
+	public DiskMetrics? GetWorkingDisk() => GetDiskMetrics(EpicGames.Horde.Agents.Metrics.DriveType.Working);
+
 	/// <inheritdoc/>
-	public DiskMetrics? GetDisk()
+	public DiskMetrics? GetSystemDisk() => GetDiskMetrics(EpicGames.Horde.Agents.Metrics.DriveType.System);
+
+	/// <summary>
+	/// Fetches root directory based off drive classification
+	/// </summary>
+	/// <param name="drive">Classification of the drive we're fetching telemetry for</param>
+	/// <returns>DiskMetrics object representing total size and free space</returns>
+	private DiskMetrics? GetDiskMetrics(EpicGames.Horde.Agents.Metrics.DriveType drive)
 	{
-		if (Statvfs(_workingDir.FullName, out StatvfsStruct stats) != 0)
+		string path = drive == EpicGames.Horde.Agents.Metrics.DriveType.Working ? _workingDir.FullName : "/";
+		if (Statvfs(path, out StatvfsStruct stats) != 0)
 		{
 			int errno = Marshal.GetLastPInvokeError();
-			throw new IOException($"Failed to get disk metrics for path '{_workingDir.FullName}'. Error: {Marshal.GetPInvokeErrorMessage(errno)}");
+			throw new IOException($"Failed to get disk metrics for path '{path}'. Error: {Marshal.GetPInvokeErrorMessage(errno)}");
 		}
 		
 		// Calculate total size and free space
diff --git a/Engine/Source/Programs/Horde/Plugins/Compute/HordeServer.Compute/Agents/AgentsController.cs b/Engine/Source/Programs/Horde/Plugins/Compute/HordeServer.Compute/Agents/AgentsController.cs
--- a/Engine/Source/Programs/Horde/Plugins/Compute/HordeServer.Compute/Agents/AgentsController.cs
+++ b/Engine/Source/Programs/Horde/Plugins/Compute/HordeServer.Compute/Agents/AgentsController.cs
@@ -468,7 +468,7 @@ namespace HordeServer.Agents
 			IReadOnlyList<IAgentTelemetry> telemetry = await _agentTelemetryCollection.FindAsync(agentId, minTime, maxTime, HttpContext.RequestAborted);
 
 			List<GetAgentTelemetrySampleResponse> samples = new List<GetAgentTelemetrySampleResponse>();
-			samples.AddRange(telemetry.Select(x => new GetAgentTelemetrySampleResponse(x.TimeUtc, x.UserCpu, x.IdleCpu, x.SystemCpu, x.FreeRam, x.UsedRam, x.TotalRam, x.FreeDisk, x.TotalDisk)));
+			samples.AddRange(telemetry.Select(x => new GetAgentTelemetrySampleResponse(x.TimeUtc, x.UserCpu, x.IdleCpu, x.SystemCpu, x.FreeRam, x.UsedRam, x.TotalRam, x.WorkingFreeDisk, x.WorkingTotalDisk, x.SystemFreeDisk, x.SystemTotalDisk)));
 
 			return new GetAgentTelemetryResponse(samples);
 		}
diff --git a/Engine/Source/Programs/Horde/Plugins/Compute/HordeServer.Compute/Agents/Telemetry/AgentTelemetryCollection.cs b/Engine/Source/Programs/Horde/Plugins/Compute/HordeServer.Compute/Agents/Telemetry/AgentTelemetryCollection.cs
--- a/Engine/Source/Programs/Horde/Plugins/Compute/HordeServer.Compute/Agents/Telemetry/AgentTelemetryCollection.cs
+++ b/Engine/Source/Programs/Horde/Plugins/Compute/HordeServer.Compute/Agents/Telemetry/AgentTelemetryCollection.cs
@@ -46,10 +46,16 @@ namespace HordeServer.Agents.Telemetry
 			public int UsedRam { get; set; }
 
 			[BsonElement("hdd_tot")]
-			public long TotalDisk { get; set; }
+			public long WorkingTotalDisk { get; set; }
 
 			[BsonElement("hdd_free")]
-			public long FreeDisk { get; set; }
+			public long WorkingFreeDisk { get; set; }
+
+			[BsonElement("sys_hdd_tot")]
+			public long SystemTotalDisk { get; set; }
+
+			[BsonElement("sys_hdd_free")]
+			public long SystemFreeDisk { get; set; }
 		}
 
 		static TimeSpan RetainDays { get; } = TimeSpan.FromDays(30.0);
@@ -124,8 +130,10 @@ namespace HordeServer.Agents.Telemetry
 			document.FreeRam = telemetry.FreeRamMb;
 			document.UsedRam = telemetry.UsedRamMb;
 			document.TotalRam = telemetry.TotalRamMb;
-			document.FreeDisk = telemetry.FreeDiskMb;
-			document.TotalDisk = telemetry.TotalDiskMb;
+			document.WorkingFreeDisk = telemetry.WorkingFreeDiskMb;
+			document.WorkingTotalDisk = telemetry.WorkingTotalDiskMb;
+			document.SystemTotalDisk = telemetry.SystemTotalDiskMb;
+			document.SystemFreeDisk = telemetry.SystemFreeDiskMb;
 			_writer.Write(document);
 		}
 
diff --git a/Engine/Source/Programs/Horde/Plugins/Compute/HordeServer.Compute/Agents/Telemetry/IAgentTelemetry.cs b/Engine/Source/Programs/Horde/Plugins/Compute/HordeServer.Compute/Agents/Telemetry/IAgentTelemetry.cs
--- a/Engine/Source/Programs/Horde/Plugins/Compute/HordeServer.Compute/Agents/Telemetry/IAgentTelemetry.cs
+++ b/Engine/Source/Programs/Horde/Plugins/Compute/HordeServer.Compute/Agents/Telemetry/IAgentTelemetry.cs
@@ -45,11 +45,21 @@ namespace HordeServer.Agents.Telemetry
 		/// <summary>
 		/// Total disk space (megabytes)
 		/// </summary>
-		long TotalDisk { get; set; }
+		long WorkingTotalDisk { get; set; }
 
 		/// <summary>
 		/// Free disk space (megabytes)
 		/// </summary>
-		long FreeDisk { get; set; }
+		long WorkingFreeDisk { get; set; }
+
+		/// <summary>
+		/// Total disk space (megabytes)
+		/// </summary>
+		long SystemTotalDisk { get; set; }
+
+		/// <summary>
+		/// Free disk space (megabytes)
+		/// </summary>
+		long SystemFreeDisk { get; set; }
 	}
 }
diff --git a/Engine/Source/Programs/Horde/Plugins/Compute/HordeServer.Compute/Agents/Telemetry/IAgentTelemetryCollection.cs b/Engine/Source/Programs/Horde/Plugins/Compute/HordeServer.Compute/Agents/Telemetry/IAgentTelemetryCollection.cs
--- a/Engine/Source/Programs/Horde/Plugins/Compute/HordeServer.Compute/Agents/Telemetry/IAgentTelemetryCollection.cs
+++ b/Engine/Source/Programs/Horde/Plugins/Compute/HordeServer.Compute/Agents/Telemetry/IAgentTelemetryCollection.cs
@@ -27,5 +27,5 @@ namespace HordeServer.Agents.Telemetry
 	/// <summary>
 	/// New telemetry data
 	/// </summary>
-	public record class NewAgentTelemetry(float UserCpuPct, float IdleCpuPct, float SystemCpuPct, int FreeRamMb, int UsedRamMb, int TotalRamMb, long FreeDiskMb, long TotalDiskMb);
+	public record class NewAgentTelemetry(float UserCpuPct, float IdleCpuPct, float SystemCpuPct, int FreeRamMb, int UsedRamMb, int TotalRamMb, long WorkingFreeDiskMb, long WorkingTotalDiskMb, long SystemFreeDiskMb, long SystemTotalDiskMb);
 }
diff --git a/Engine/Source/Programs/Horde/Plugins/Compute/HordeServer.Compute/Server/RpcService.cs b/Engine/Source/Programs/Horde/Plugins/Compute/HordeServer.Compute/Server/RpcService.cs
--- a/Engine/Source/Programs/Horde/Plugins/Compute/HordeServer.Compute/Server/RpcService.cs
+++ b/Engine/Source/Programs/Horde/Plugins/Compute/HordeServer.Compute/Server/RpcService.cs
@@ -423,7 +423,11 @@ namespace HordeServer.Server
 		{
 			_logger.LogDebug("Posting telemetry data for {AgentId}", new AgentId(request.AgentId));
 
-			NewAgentTelemetry telemetry = new NewAgentTelemetry(request.UserCpu, request.IdleCpu, request.SystemCpu, (int)request.FreeRam, (int)request.UsedRam, (int)request.TotalRam, (long)request.FreeDisk, (long)request.TotalDisk);
+			NewAgentTelemetry telemetry = new NewAgentTelemetry(
+				request.UserCpu, request.IdleCpu, request.SystemCpu, 
+				(int)request.FreeRam, (int)request.UsedRam, (int)request.TotalRam, 
+				(long)request.WorkingFreeDisk, (long)request.WorkingTotalDisk,
+				(long)request.SystemFreeDisk, (long)request.SystemTotalDisk);
 			_agentTelemetryCollection.Add(new AgentId(request.AgentId), telemetry);
 
 			return Task.FromResult(new Empty());
diff --git a/Engine/Source/Programs/Shared/EpicGames.Horde/Agents/AgentMessages.cs b/Engine/Source/Programs/Shared/EpicGames.Horde/Agents/AgentMessages.cs
--- a/Engine/Source/Programs/Shared/EpicGames.Horde/Agents/AgentMessages.cs
+++ b/Engine/Source/Programs/Shared/EpicGames.Horde/Agents/AgentMessages.cs
@@ -133,5 +133,5 @@ namespace EpicGames.Horde.Agents
 	/// <summary>
 	/// Telemetry data sample
 	/// </summary>
-	public record GetAgentTelemetrySampleResponse(DateTime Time, float UserCpu, float IdleCpu, float SystemCpu, int FreeRam, int UsedRam, int TotalRam, long FreeDisk, long TotalDisk);
+	public record GetAgentTelemetrySampleResponse(DateTime Time, float UserCpu, float IdleCpu, float SystemCpu, int FreeRam, int UsedRam, int TotalRam, long FreeDisk, long TotalDisk, long SystemFreeDisk, long SystemTotalDisk);
 }
diff --git a/Engine/Source/Programs/Shared/EpicGames.Horde/Agents/Metrics/Metrics.cs b/Engine/Source/Programs/Shared/EpicGames.Horde/Agents/Metrics/Metrics.cs
--- a/Engine/Source/Programs/Shared/EpicGames.Horde/Agents/Metrics/Metrics.cs
+++ b/Engine/Source/Programs/Shared/EpicGames.Horde/Agents/Metrics/Metrics.cs
@@ -144,8 +144,29 @@ public interface ISystemMetrics
 	MemoryMetrics? GetMemory();
 
 	/// <summary>
-	/// Gets HDD metrics
+	/// Gets Working HDD metrics
 	/// </summary>
 	/// <returns>An object with disk usage metrics</returns>
-	DiskMetrics? GetDisk();
+	DiskMetrics? GetWorkingDisk();
+
+	/// <summary>
+	/// Gets System HDD metrics
+	/// </summary>
+	/// <returns>An object with disk usage metrics</returns>
+	DiskMetrics? GetSystemDisk();
 }
+
+/// <summary>
+/// Enum to handle separate drive types tracked by Horde telemetry
+/// </summary>
+public enum DriveType
+{
+	/// <summary>
+	/// System-level drive
+	/// </summary>
+	System = 0,
+	/// <summary>
+	/// Horde Working drive
+	/// </summary>
+	Working = 1
+}
\ No newline at end of file
