Horde: Add optional helpSlack/helpEmail overrides to ErrorInfo for per-error support contact customization

#Synced-CL 0

diff --git a/Engine/Source/Programs/Horde/HordeDashboard/src/components/error/ErrorBar.tsx b/Engine/Source/Programs/Horde/HordeDashboard/src/components/error/ErrorBar.tsx
--- a/Engine/Source/Programs/Horde/HordeDashboard/src/components/error/ErrorBar.tsx
+++ b/Engine/Source/Programs/Horde/HordeDashboard/src/components/error/ErrorBar.tsx
@@ -9,8 +9,8 @@ import dashboard from 'horde/backend/Dashboard';
 export const ErrorBar: React.FC = observer(() => {
     const {error, clear} = errorBarStore;
 
-    const helpEmail = dashboard.helpEmail;
-    const helpSlack = dashboard.helpSlack;
+    const helpEmail = error?.helpEmail ?? dashboard.helpEmail;
+    const helpSlack = error?.helpSlack ?? dashboard.helpSlack;
     const anyHelp = !!helpSlack || !!helpEmail;
 
     const getConsoleErrorMessage = (): string => {
diff --git a/Engine/Source/Programs/Horde/HordeDashboard/src/components/error/ErrorDialog.tsx b/Engine/Source/Programs/Horde/HordeDashboard/src/components/error/ErrorDialog.tsx
--- a/Engine/Source/Programs/Horde/HordeDashboard/src/components/error/ErrorDialog.tsx
+++ b/Engine/Source/Programs/Horde/HordeDashboard/src/components/error/ErrorDialog.tsx
@@ -32,8 +32,8 @@ export const ErrorDialog: React.FC = observer(() => {
       title: title
    };
 
-   const helpEmail = dashboard.helpEmail;
-   const helpSlack = dashboard.helpSlack;
+   const helpEmail = error.helpEmail ?? dashboard.helpEmail;
+   const helpSlack = error.helpSlack ?? dashboard.helpSlack;
 
    let fontSize = 13;
 
diff --git a/Engine/Source/Programs/Horde/HordeDashboard/src/components/error/ErrorStore.ts b/Engine/Source/Programs/Horde/HordeDashboard/src/components/error/ErrorStore.ts
--- a/Engine/Source/Programs/Horde/HordeDashboard/src/components/error/ErrorStore.ts
+++ b/Engine/Source/Programs/Horde/HordeDashboard/src/components/error/ErrorStore.ts
@@ -35,6 +35,10 @@ export type ErrorInfo = {
 
    format?: ErrorFormatted;
 
+   // custom help contact overrides (falls back to global dashboard config if not set)
+   helpSlack?: string;
+   helpEmail?: string;
+
 }
 
 export enum ErrorReportType {
