Horde: Dispose Perforce PoolConnectionHandle when there is an exception during login

#Synced-CL 0

diff --git a/Engine/Source/Programs/Horde/Plugins/Build/HordeServer.Build/VersionControl/Perforce/PerforceService.cs b/Engine/Source/Programs/Horde/Plugins/Build/HordeServer.Build/VersionControl/Perforce/PerforceService.cs
--- a/Engine/Source/Programs/Horde/Plugins/Build/HordeServer.Build/VersionControl/Perforce/PerforceService.cs
+++ b/Engine/Source/Programs/Horde/Plugins/Build/HordeServer.Build/VersionControl/Perforce/PerforceService.cs
@@ -470,7 +470,16 @@ namespace HordeServer.VersionControl.Perforce
 
 				if (!String.IsNullOrEmpty(credentials.Password) && String.IsNullOrEmpty(credentials.Ticket))
 				{
-					await handle.LoginAsync(credentials.Password, cancellationToken);
+					try
+					{
+						await handle.LoginAsync(credentials.Password, cancellationToken);
+					}
+					catch (Exception ex)
+					{
+						_logger.LogError(ex, "Unable to login: {Message}", ex.Message);
+						handle.Dispose();
+						throw;
+					}
 				}
 			}
 			return handle;
