Horde: Fix build health compiler warnings

diff --git a/Engine/Source/Programs/Horde/Plugins/Build/HordeServer.Build.Tests/Issues/IssueServiceTests.cs b/Engine/Source/Programs/Horde/Plugins/Build/HordeServer.Build.Tests/Issues/IssueServiceTests.cs
--- a/Engine/Source/Programs/Horde/Plugins/Build/HordeServer.Build.Tests/Issues/IssueServiceTests.cs
+++ b/Engine/Source/Programs/Horde/Plugins/Build/HordeServer.Build.Tests/Issues/IssueServiceTests.cs
@@ -3732,7 +3732,7 @@ namespace HordeServer.Tests.Issues
 
 			IReadOnlyList<IIssue> parallelIssues = await IssueCollection.FindIssuesAsync();
 			// Filter to the new issue
-			IIssue? parallelIssue = parallelIssues.FirstOrDefault(i => i.Fingerprints.Any(f => f.Keys.Any(k => k.Name.Contains("errorfile2"))));
+			IIssue? parallelIssue = parallelIssues.FirstOrDefault(i => i.Fingerprints.Any(f => f.Keys.Any(k => k.Name.Contains("errorfile2", StringComparison.Ordinal))));
 			Assert.IsNotNull(parallelIssue, "Parallel issue should be created");
 
 			IReadOnlyList<IIssueSuspect> parallelSuspects = await IssueCollection.FindSuspectsAsync(parallelIssue);
@@ -3777,7 +3777,7 @@ namespace HordeServer.Tests.Issues
 			await UpdateCompleteStepAsync(job2, 0, 0, JobStepOutcome.Failure);
 
 			IReadOnlyList<IIssue> issues = await IssueCollection.FindIssuesAsync();
-			IIssue? issue = issues.FirstOrDefault(i => i.Fingerprints.Any(f => f.Keys.Any(k => k.Name.Contains("file815"))));
+			IIssue? issue = issues.FirstOrDefault(i => i.Fingerprints.Any(f => f.Keys.Any(k => k.Name.Contains("file815", StringComparison.Ordinal))));
 			Assert.IsNotNull(issue, "Issue should be created");
 
 			IReadOnlyList<IIssueSuspect> suspects = await IssueCollection.FindSuspectsAsync(issue);
@@ -3819,7 +3819,7 @@ namespace HordeServer.Tests.Issues
 			await UpdateCompleteStepAsync(job2, 0, 0, JobStepOutcome.Failure);
 
 			IReadOnlyList<IIssue> issues = await IssueCollection.FindIssuesAsync();
-			IIssue? issue = issues.FirstOrDefault(i => i.Fingerprints.Any(f => f.Keys.Any(k => k.Name.Contains("earlyexit"))));
+			IIssue? issue = issues.FirstOrDefault(i => i.Fingerprints.Any(f => f.Keys.Any(k => k.Name.Contains("earlyexit", StringComparison.Ordinal))));
 			Assert.IsNotNull(issue, "Issue should be created");
 
 			IReadOnlyList<IIssueSuspect> suspects = await IssueCollection.FindSuspectsAsync(issue);
@@ -3857,7 +3857,7 @@ namespace HordeServer.Tests.Issues
 			await UpdateCompleteStepAsync(job2, 0, 0, JobStepOutcome.Failure);
 
 			IReadOnlyList<IIssue> initialIssues = await IssueCollection.FindIssuesAsync();
-			IIssue? originalIssue = initialIssues.FirstOrDefault(i => i.Fingerprints.Any(f => f.Keys.Any(k => k.Name.Contains("fallbacktest"))));
+			IIssue? originalIssue = initialIssues.FirstOrDefault(i => i.Fingerprints.Any(f => f.Keys.Any(k => k.Name.Contains("fallbacktest", StringComparison.Ordinal))));
 			Assert.IsNotNull(originalIssue, "Original issue should be created");
 			int originalIssueId = originalIssue.Id;
 
@@ -3879,7 +3879,7 @@ namespace HordeServer.Tests.Issues
 
 			// The fingerprint fallback should have matched the existing issue instead of creating a new one
 			IReadOnlyList<IIssue> finalIssues = await IssueCollection.FindIssuesAsync();
-			IReadOnlyList<IIssue> matchingIssues = finalIssues.Where(i => i.Fingerprints.Any(f => f.Keys.Any(k => k.Name.Contains("fallbacktest")))).ToList();
+			IReadOnlyList<IIssue> matchingIssues = finalIssues.Where(i => i.Fingerprints.Any(f => f.Keys.Any(k => k.Name.Contains("fallbacktest", StringComparison.Ordinal)))).ToList();
 
 			// With fingerprint fallback, should still only have 1 issue (the original)
 			// because the new span should have been added to the existing issue
@@ -4046,7 +4046,7 @@ namespace HordeServer.Tests.Issues
 			await UpdateCompleteStepAsync(job2, 0, 0, JobStepOutcome.Failure);
 
 			IReadOnlyList<IIssue> issues = await IssueCollection.FindIssuesAsync();
-			IIssue? issue = issues.FirstOrDefault(i => i.Fingerprints.Any(f => f.Keys.Any(k => k.Name.Contains("combined1625"))));
+			IIssue? issue = issues.FirstOrDefault(i => i.Fingerprints.Any(f => f.Keys.Any(k => k.Name.Contains("combined1625", StringComparison.Ordinal))));
 			Assert.IsNotNull(issue, "Issue should be created with combined features");
 
 			IReadOnlyList<IIssueSuspect> suspects = await IssueCollection.FindSuspectsAsync(issue);
@@ -4092,7 +4092,7 @@ namespace HordeServer.Tests.Issues
 			await UpdateCompleteStepAsync(job2, 0, 0, JobStepOutcome.Failure);
 
 			IReadOnlyList<IIssue> issues = await IssueCollection.FindIssuesAsync();
-			IIssue? issue = issues.FirstOrDefault(i => i.Fingerprints.Any(f => f.Keys.Any(k => k.Name.Contains("integration1730"))));
+			IIssue? issue = issues.FirstOrDefault(i => i.Fingerprints.Any(f => f.Keys.Any(k => k.Name.Contains("integration1730", StringComparison.Ordinal))));
 			Assert.IsNotNull(issue, "Issue should be created with all features enabled");
 
 			IReadOnlyList<IIssueSuspect> suspects = await IssueCollection.FindSuspectsAsync(issue);
