Horde: Remove full URL from OpenTelemetry trace name in Horde Agent

* It has too high cardinality when IDs, CLs etc get included
* Added support for passing along a request option called "url.template" to let calling code provide a better name hint

diff --git a/Engine/Source/Programs/Horde/HordeAgent/Utility/OpenTelemetry.cs b/Engine/Source/Programs/Horde/HordeAgent/Utility/OpenTelemetry.cs
--- a/Engine/Source/Programs/Horde/HordeAgent/Utility/OpenTelemetry.cs
+++ b/Engine/Source/Programs/Horde/HordeAgent/Utility/OpenTelemetry.cs
@@ -41,6 +41,7 @@ public class OpenTelemetryDatadogLogEnricher : ILogEventEnricher
 public static class OpenTelemetryHelper
 {
 	private static ResourceBuilder? s_resourceBuilder;
+	private static readonly HttpRequestOptionsKey<string> s_urlTemplateKey = new("url.template");
 	
 	/// <summary>
 	/// Name of resource name attribute in Datadog
@@ -88,12 +89,19 @@ public static class OpenTelemetryHelper
 
 	private static void ConfigureTracing(TracerProviderBuilder builder, OpenTelemetrySettings settings)
 	{
-		void DatadogHttpRequestEnricher(Activity activity, HttpRequestMessage message)
+		void DatadogHttpRequestEnricher(Activity activity, HttpRequestMessage request)
 		{
 			activity.SetTag("service.name", settings.ServiceName + "-http-client");
 			activity.SetTag("operation.name", "http.request");
-			string url = $"{message.Method} {message.Headers.Host}{message.RequestUri?.LocalPath}";
-			activity.DisplayName = url;
+			string url = $"{request.Method} {request.Headers.Host}{request.RequestUri?.LocalPath}";
+			if (request.Options.TryGetValue(s_urlTemplateKey, out string? template))
+			{
+				activity.DisplayName = $"{request.Method} {template}";
+			}
+			else
+			{
+				activity.DisplayName = $"{request.Method} {request.Headers.Host}";
+			}
 			activity.SetTag("resource.name", url);
 		}
 
