[Horde Server][Horde Dashboard][Analytics] Adding system drive stats to CapabilitiesService + displaying new drive stats on Horde Dashboard

- Added new set of system drive related disk info events propagated through the CapabilitiesService (`C:\` on Windows and `/` on Unix).
- Added HordeDashboard boilerplate to support these new analytics events, and modified nomenclature of pre-existing "working drive" telemetry

#Synced-CL 48478172

diff --git a/Engine/Source/Programs/Horde/HordeAgent/Services/CapabilitiesService.cs b/Engine/Source/Programs/Horde/HordeAgent/Services/CapabilitiesService.cs
--- a/Engine/Source/Programs/Horde/HordeAgent/Services/CapabilitiesService.cs
+++ b/Engine/Source/Programs/Horde/HordeAgent/Services/CapabilitiesService.cs
@@ -407,6 +407,26 @@ namespace HordeAgent.Services
 			capabilities.Properties.Add($"BootTime={s_bootTime}");
 			capabilities.Properties.Add($"StartTime={s_startTime}");
 
+			// Add system disk info
+			try
+			{
+				string? systemDriveName;
+				systemDriveName = RuntimeInformation.IsOSPlatform(OSPlatform.Windows)
+					? Path.GetPathRoot(Environment.SystemDirectory) // Windows (C:\)
+					: "/"; // *nix
+
+				if (systemDriveName != null)
+				{
+					DriveInfo systemInfo = new DriveInfo(systemDriveName);
+					capabilities.Properties.Add($"{KnownPropertyNames.SystemDiskFreeSpace}={systemInfo.AvailableFreeSpace}");
+					capabilities.Properties.Add($"{KnownPropertyNames.SystemDiskTotalSize}={systemInfo.TotalSize}");
+				}
+			}
+			catch (Exception ex)
+			{
+				logger.LogWarning(ex, "Unable to query system disk info");
+			}
+
 			// Add information about the current session
 			if (workingDir != null)
 			{
@@ -426,8 +446,8 @@ namespace HordeAgent.Services
 					try
 					{
 						DriveInfo info = new DriveInfo(driveName);
-						capabilities.Properties.Add($"{KnownPropertyNames.DiskFreeSpace}={info.AvailableFreeSpace}");
-						capabilities.Properties.Add($"{KnownPropertyNames.DiskTotalSize}={info.TotalSize}");
+						capabilities.Properties.Add($"{KnownPropertyNames.WorkingDiskFreeSpace}={info.AvailableFreeSpace}");
+						capabilities.Properties.Add($"{KnownPropertyNames.WorkingDiskTotalSize}={info.TotalSize}");
 					}
 					catch (Exception ex)
 					{
diff --git a/Engine/Source/Programs/Horde/HordeDashboard/src/backend/Api.ts b/Engine/Source/Programs/Horde/HordeDashboard/src/backend/Api.ts
--- a/Engine/Source/Programs/Horde/HordeDashboard/src/backend/Api.ts
+++ b/Engine/Source/Programs/Horde/HordeDashboard/src/backend/Api.ts
@@ -2189,8 +2189,10 @@ export type GetAgentTelemetrySampleResponse = {
 	freeRam: number;
 	usedRam: number;
 	totalRam: number;
-	freeDisk: number;
-	totalDisk: number;
+	workingFreeDisk: number;
+	workingTotalDisk: number;
+	systemFreeDisk: number;
+	systemTotalDisk: number;
 }
 
 /**Information about a workspace type */
diff --git a/Engine/Source/Programs/Horde/HordeDashboard/src/components/agents/AgentTelemetrySparkline.tsx b/Engine/Source/Programs/Horde/HordeDashboard/src/components/agents/AgentTelemetrySparkline.tsx
--- a/Engine/Source/Programs/Horde/HordeDashboard/src/components/agents/AgentTelemetrySparkline.tsx
+++ b/Engine/Source/Programs/Horde/HordeDashboard/src/components/agents/AgentTelemetrySparkline.tsx
@@ -238,8 +238,8 @@ export const AgentTelemetrySparkline: React.FC<{ handler: AgentTelemetryHandler
    })
 
    const diskData = data.map(d => {
-      if (typeof (d.totalDisk) === "number" && d.totalDisk) {
-         return 1.0 - (d.freeDisk / d.totalDisk)
+      if (typeof (d.workingTotalDisk) === "number" && d.workingTotalDisk) {
+         return 1.0 - (d.workingFreeDisk / d.workingTotalDisk)
       }
 
       return 0;
diff --git a/Engine/Source/Programs/Horde/HordeDashboard/src/components/agents/HistoryModal.tsx b/Engine/Source/Programs/Horde/HordeDashboard/src/components/agents/HistoryModal.tsx
--- a/Engine/Source/Programs/Horde/HordeDashboard/src/components/agents/HistoryModal.tsx
+++ b/Engine/Source/Programs/Horde/HordeDashboard/src/components/agents/HistoryModal.tsx
@@ -37,7 +37,7 @@ const historyStyles = mergeStyleSets({
 
 const state = historyModalStore;
 
-const TelemetryTooltip: React.FC<{ id: "cpu" | "ram" | "disk", cpuKey?: string }> = observer(({ id, cpuKey }) => {
+const TelemetryTooltip: React.FC<{ id: "cpu" | "ram" | "workingDisk" | "systemDisk", cpuKey?: string }> = observer(({ id, cpuKey }) => {
 
    const { modeColors } = getHordeStyling();
 
@@ -161,9 +161,13 @@ const TelemetryTooltip: React.FC<{ id: "cpu" | "ram" | "disk", cpuKey?: string }
             if (cpuKey == "system") value = `${Math.ceil(d.systemCpu)}%`;
             if (cpuKey == "combined") value = `${Math.ceil(d.userCpu + d.systemCpu)}%`;
             break;
-         case "disk":
+         case "workingDisk":
             title = "Disk"
-            value = formatBytes((d.totalDisk - d.freeDisk) * 1024 * 1024) + " Used";
+            value = formatBytes((d.workingTotalDisk - d.workingFreeDisk) * 1024 * 1024) + " Used";
+            break;
+         case "systemDisk":
+            title = "Disk"
+            value = formatBytes((d.systemTotalDisk - d.systemFreeDisk) * 1024 * 1024) + " Used";
             break;
          case "ram":
             title = "Ram"
@@ -378,23 +382,47 @@ const TelemetryPanel: React.FC<{}> = observer(({ }) => {
    const minTime = hasData ? getShortNiceTime(min, true) : "";
    const maxTime = hasData ? getShortNiceTime(max, true) : "";
 
-   const diskData = data.map(d => {
-      if (typeof (d.totalDisk) === "number" && d.totalDisk) {
-         return 1.0 - (d.freeDisk / d.totalDisk)
-      }
+   const calculateDiskUtilization = (
+      data: GetAgentTelemetrySampleResponse[], 
+      totalKey: keyof GetAgentTelemetrySampleResponse, 
+      freeKey: keyof GetAgentTelemetrySampleResponse
+   ) => 
+      data.map(d => {
+         const total = d[totalKey];
+         const free = d[freeKey];
+
+         if (typeof total == "number" && typeof free == "number" && total > 0) {
+            return 1.0 - (free / total);
+         }
+   
+         return 0;
+      })
 
-      return 0;
-   })
+   const diskData = calculateDiskUtilization(data, 'workingTotalDisk', 'workingFreeDisk');
+   const systemDiskData = calculateDiskUtilization(data, 'systemTotalDisk', 'systemFreeDisk');
 
    let cpuText = "";
    let ramText = "";
    let diskText = "";
+   let systemDiskText = "";
 
    let cpuLabel = "";
    if (cpuKey == "user") cpuLabel = "CPU - User";
    if (cpuKey == "system") cpuLabel = "CPU - System";
    if (cpuKey == "combined") cpuLabel = "CPU - Combined";
 
+   const formatDiskSize = (bytes: number, decimals = 0) => {
+      if (!+bytes) return '0 Bytes'
+
+      const k = 1024
+      const dm = decimals < 0 ? 0 : decimals
+      const sizes = ['Bytes', 'KB', 'MB', 'GB', 'TB', 'PB', 'EB', 'ZB', 'YB']
+
+      const i = Math.floor(Math.log(bytes) / Math.log(k))
+
+      return `${parseFloat((bytes / Math.pow(k, i)).toFixed(dm))} ${sizes[i]}`
+   }
+
    const devices = state.selectedAgent?.capabilities?.devices;
    if (devices?.length) {
       const device = devices[0];
@@ -407,23 +435,17 @@ const TelemetryPanel: React.FC<{}> = observer(({ }) => {
             ramText = `${v.replace("RAM=", "")} GB`;
          }
 
-         if (v.startsWith("DiskTotalSize=")) {
+         if (v.startsWith("WorkingDiskTotalSize=")) {
             const elements = v.split("=");
             if (elements.length === 2) {
+               diskText = formatDiskSize(parseInt(elements[1]))
+            }
+         }
 
-               function formatBytes(bytes: number, decimals = 0) {
-                  if (!+bytes) return '0 Bytes'
-
-                  const k = 1024
-                  const dm = decimals < 0 ? 0 : decimals
-                  const sizes = ['Bytes', 'KB', 'MB', 'GB', 'TB', 'PB', 'EB', 'ZB', 'YB']
-
-                  const i = Math.floor(Math.log(bytes) / Math.log(k))
-
-                  return `${parseFloat((bytes / Math.pow(k, i)).toFixed(dm))} ${sizes[i]}`
-               }
-
-               diskText = formatBytes(parseInt(elements[1]))
+         if (v.startsWith("SystemDiskTotalSize=")) {
+            const elements = v.split("=");
+            if (elements.length === 2) {
+               systemDiskText = formatDiskSize(parseInt(elements[1]))
             }
          }
 
@@ -613,9 +635,9 @@ const TelemetryPanel: React.FC<{}> = observer(({ }) => {
 
       {!!diskText && <Stack verticalAlign="center" verticalFill>
          <Stack>
-            <Label>Storage Utilization - {diskText}</Label>
+            <Label>Agent Working Drive Utilization - {diskText}</Label>
          </Stack>
-         <div onMouseLeave={() => state.setLeaseTooltip(undefined)} onMouseMove={(ev) => { relativeCoords(ev, "disk") }} onClick={() => {
+         <div onMouseLeave={() => state.setLeaseTooltip(undefined)} onMouseMove={(ev) => { relativeCoords(ev, "workingDisk") }} onClick={() => {
 
             if (!state.currentLeaseTooltip?.frozen) {
                if (state.currentLeaseTooltip) {
@@ -624,7 +646,7 @@ const TelemetryPanel: React.FC<{}> = observer(({ }) => {
             }
          }}>
             <div style={{ position: "relative" }}>
-               <TelemetryTooltip id="disk" />
+               <TelemetryTooltip id="workingDisk" />
             </div>
             <div style={{ backgroundColor: dashboard.darktheme ? "#060709" : "#F3F2F1", padding: 8, border: "solid 1px #181A1B" }}>
                <div style={{ position: "relative" }}>
@@ -636,7 +658,39 @@ const TelemetryPanel: React.FC<{}> = observer(({ }) => {
                      <SparklinesReferenceLine type="custom" value={96} style={{ stroke: statusColors.get(StatusColor.Success) }} />
                   </Sparklines>
                   <svg style={{ position: "absolute", top: 0, left: 0, pointerEvents: "none" }} width={sparkWidth} height={sparkHeight} viewBox={`0 0 ${sparkWidth} ${sparkHeight}`} preserveAspectRatio="none">
-                     <LeaseOverlayRects processedLeases={processedLeases} sparkHeight={sparkHeight} sparkWidth={sparkWidth} chartId="disk" />
+                     <LeaseOverlayRects processedLeases={processedLeases} sparkHeight={sparkHeight} sparkWidth={sparkWidth} chartId="workingDisk" />
+                  </svg>
+               </div>
+            </div>
+         </div>
+      </Stack>}
+
+      {!!systemDiskText && <Stack verticalAlign="center" verticalFill>
+         <Stack>
+            <Label>System Drive Utilization - {systemDiskText}</Label>
+         </Stack>
+         <div onMouseLeave={() => state.setLeaseTooltip(undefined)} onMouseMove={(ev) => { relativeCoords(ev, "systemDisk") }} onClick={() => {
+
+            if (!state.currentLeaseTooltip?.frozen) {
+               if (state.currentLeaseTooltip) {
+                  state.setLeaseTooltip({ ...state.currentLeaseTooltip, frozen: true });
+               }
+            }
+         }}>
+            <div style={{ position: "relative" }}>
+               <TelemetryTooltip id="systemDisk" />
+            </div>
+            <div style={{ backgroundColor: dashboard.darktheme ? "#060709" : "#F3F2F1", padding: 8, border: "solid 1px #181A1B" }}>
+               <div style={{ position: "relative" }}>
+                  <Sparklines width={sparkWidth} height={sparkHeight} svgWidth={sparkWidth} svgHeight={sparkHeight} data={systemDiskData} min={0} max={1}>
+                     <SparklinesLine color={dashboard.darktheme ? "lightblue" : "#1E90FF"} />
+                     <SparklinesReferenceLine type="custom" value={0} style={{ stroke: statusColors.get(StatusColor.Failure) }} />
+                     <SparklinesReferenceLine type="custom" value={32} style={{ stroke: statusColors.get(StatusColor.Warnings) }} />
+                     <SparklinesReferenceLine type="custom" value={64} style={{ stroke: statusColors.get(StatusColor.Success) }} />
+                     <SparklinesReferenceLine type="custom" value={96} style={{ stroke: statusColors.get(StatusColor.Success) }} />
+                  </Sparklines>
+                  <svg style={{ position: "absolute", top: 0, left: 0, pointerEvents: "none" }} width={sparkWidth} height={sparkHeight} viewBox={`0 0 ${sparkWidth} ${sparkHeight}`} preserveAspectRatio="none">
+                     <LeaseOverlayRects processedLeases={processedLeases} sparkHeight={sparkHeight} sparkWidth={sparkWidth} chartId="systemDisk" />
                   </svg>
                </div>
             </div>
diff --git a/Engine/Source/Programs/Horde/Plugins/Build/HordeServer.Build.Tests/Agents/Pools/PoolUpdateServiceTest.cs b/Engine/Source/Programs/Horde/Plugins/Build/HordeServer.Build.Tests/Agents/Pools/PoolUpdateServiceTest.cs
--- a/Engine/Source/Programs/Horde/Plugins/Build/HordeServer.Build.Tests/Agents/Pools/PoolUpdateServiceTest.cs
+++ b/Engine/Source/Programs/Horde/Plugins/Build/HordeServer.Build.Tests/Agents/Pools/PoolUpdateServiceTest.cs
@@ -150,7 +150,7 @@ public class PoolUpdateServiceTest : BuildTestSetup
 			x.Plugins.GetBuildConfig().Projects.Add(new ProjectConfig { Streams = [new StreamConfig { WorkspaceTypes = workspaceTypes, AgentTypes = agentTypes }]});
 		});
 		
-		IAgent? agent = await CreateAgentAsync(_pool, properties: [$"{KnownPropertyNames.DiskFreeSpace}={freeDiskSpaceMb * 1024 * 1024}"]);
+		IAgent? agent = await CreateAgentAsync(_pool, properties: [$"{KnownPropertyNames.WorkingDiskFreeSpace}={freeDiskSpaceMb * 1024 * 1024}"]);
 		agent = await agent.TryTerminateSessionAsync(); // Make agent status = stopped
 		Assert.IsNotNull(agent);
 		return agent;
diff --git a/Engine/Source/Programs/Horde/Plugins/Build/HordeServer.Build.Tests/Jobs/JobTaskSourceTests.cs b/Engine/Source/Programs/Horde/Plugins/Build/HordeServer.Build.Tests/Jobs/JobTaskSourceTests.cs
--- a/Engine/Source/Programs/Horde/Plugins/Build/HordeServer.Build.Tests/Jobs/JobTaskSourceTests.cs
+++ b/Engine/Source/Programs/Horde/Plugins/Build/HordeServer.Build.Tests/Jobs/JobTaskSourceTests.cs
@@ -196,7 +196,7 @@ namespace HordeServer.Tests.Jobs
 				config.Plugins.GetBuildConfig().EnableJobTaskConformCheck = true;
 			});
 
-			List<string> properties = [$"{KnownPropertyNames.DiskFreeSpace}={diskFreeSpaceInMb * 1024 * 1024}"];
+			List<string> properties = [$"{KnownPropertyNames.WorkingDiskFreeSpace}={diskFreeSpaceInMb * 1024 * 1024}"];
 			agent = Deref(await AgentService.CreateSessionAsync(agent, new RpcAgentCapabilities(properties), null));
 
 			await JobTaskSource.TickAsync(CancellationToken.None);
diff --git a/Engine/Source/Programs/Horde/Plugins/Compute/HordeServer.Compute/Agents/IAgent.cs b/Engine/Source/Programs/Horde/Plugins/Compute/HordeServer.Compute/Agents/IAgent.cs
--- a/Engine/Source/Programs/Horde/Plugins/Compute/HordeServer.Compute/Agents/IAgent.cs
+++ b/Engine/Source/Programs/Horde/Plugins/Compute/HordeServer.Compute/Agents/IAgent.cs
@@ -504,7 +504,7 @@ namespace HordeServer.Agents
 		/// <returns>Amount of free disk space in bytes</returns>
 		public static long? GetDiskFreeSpace(this IAgent agent)
 		{
-			List<string> values = agent.GetPropertyValues(KnownPropertyNames.DiskFreeSpace).ToList();
+			List<string> values = agent.GetPropertyValues(KnownPropertyNames.WorkingDiskFreeSpace).ToList();
 			return values.Count > 0 && Int64.TryParse(values[0], out long amount) ? amount : null;
 		}
 
diff --git a/Engine/Source/Programs/Shared/EpicGames.Horde/Agents/AgentMessages.cs b/Engine/Source/Programs/Shared/EpicGames.Horde/Agents/AgentMessages.cs
--- a/Engine/Source/Programs/Shared/EpicGames.Horde/Agents/AgentMessages.cs
+++ b/Engine/Source/Programs/Shared/EpicGames.Horde/Agents/AgentMessages.cs
@@ -133,5 +133,5 @@ namespace EpicGames.Horde.Agents
 	/// <summary>
 	/// Telemetry data sample
 	/// </summary>
-	public record GetAgentTelemetrySampleResponse(DateTime Time, float UserCpu, float IdleCpu, float SystemCpu, int FreeRam, int UsedRam, int TotalRam, long FreeDisk, long TotalDisk, long SystemFreeDisk, long SystemTotalDisk);
+	public record GetAgentTelemetrySampleResponse(DateTime Time, float UserCpu, float IdleCpu, float SystemCpu, int FreeRam, int UsedRam, int TotalRam, long WorkingFreeDisk, long WorkingTotalDisk, long SystemFreeDisk, long SystemTotalDisk);
 }
diff --git a/Engine/Source/Programs/Shared/EpicGames.Horde/Agents/KnownPropertyNames.cs b/Engine/Source/Programs/Shared/EpicGames.Horde/Agents/KnownPropertyNames.cs
--- a/Engine/Source/Programs/Shared/EpicGames.Horde/Agents/KnownPropertyNames.cs
+++ b/Engine/Source/Programs/Shared/EpicGames.Horde/Agents/KnownPropertyNames.cs
@@ -53,14 +53,24 @@ public static class KnownPropertyNames
 	public const string RequestedPools = "RequestedPools";
 
 	/// <summary>
-	/// The total size of storage space on drive, in bytes
+	/// The total size of the working drive, in bytes
 	/// </summary>
-	public const string DiskTotalSize = "DiskTotalSize";
+	public const string WorkingDiskTotalSize = "WorkingDiskTotalSize";
 
 	/// <summary>
-	/// Amount of available free space on drive, in bytes
+	/// Amount of available free space on the working drive, in bytes
 	/// </summary>
-	public const string DiskFreeSpace = "DiskFreeSpace";
+	public const string WorkingDiskFreeSpace = "WorkingDiskFreeSpace";
+
+	/// <summary>
+	/// The total size of the system drive, in bytes
+	/// </summary>
+	public const string SystemDiskTotalSize = "SystemDiskTotalSize";
+
+	/// <summary>
+	/// Amount of available free space on the system drive, in bytes
+	/// </summary>
+	public const string SystemDiskFreeSpace = "SystemDiskFreeSpace";
 
 	/// <summary>
 	/// IP address used for sending compute task payloads
