[Horde Server][Analytics] Propagating full agent telemetry events to associated data sinks

diff --git a/Engine/Source/Programs/Horde/Plugins/Compute/HordeServer.Compute/Agents/Telemetry/AgentTelemetry.cs b/Engine/Source/Programs/Horde/Plugins/Compute/HordeServer.Compute/Agents/Telemetry/AgentTelemetry.cs
new file mode 100644
--- /dev/null
+++ b/Engine/Source/Programs/Horde/Plugins/Compute/HordeServer.Compute/Agents/Telemetry/AgentTelemetry.cs
@@ -0,0 +1,136 @@
+// Copyright Epic Games, Inc. All Rights Reserved.
+
+using System.ComponentModel.DataAnnotations.Schema;
+using EpicGames.Analytics;
+using EpicGames.Analytics.Telemetry;
+using EpicGames.Horde.Agents;
+using HordeServer.Agents.Telemetry;
+
+namespace HordeServer.Agents
+{
+	/// <summary>
+	/// Record used for Agent telemetry.
+	/// </summary>
+	[AnalyticsTableGen]
+	[Table("ingest.horde.state_agent_telemetry")]
+	public record AgentTelemetry : AbstractTelemetryRecord
+	{
+		/// <summary>
+		/// The ID of the agent we're reporting telemetry from.
+		/// </summary>
+		[Column("agent_id")]
+		public AgentId AgentId { get; init; }
+
+		/// <summary>
+		/// The user cpu usage.
+		/// </summary>
+		[Column("user_cpu")]
+		public float UserCpu { get; init; }
+
+		/// <summary>
+		/// The idle cpu usage.
+		/// </summary>
+		[Column("idle_cpu")]
+		public float IdleCpu { get; init; }
+
+		/// <summary>
+		/// The system cpu usage.
+		/// </summary>
+		[Column("system_cpu")]
+		public float SystemCpu { get; init; }
+
+		/// <summary>
+		/// Free memory in MB.
+		/// </summary>
+		[Column("free_ram_mb")]
+		public int FreeRamMb { get; init; }
+
+		/// <summary>
+		/// Used memory in MB.
+		/// </summary>
+		[Column("used_ram_mb")]
+		public int UsedRamMb { get; init; }
+
+		/// <summary>
+		/// Total memory in MB.
+		/// </summary>
+		[Column("total_ram_mb")]
+		public int TotalRamMb { get; init; }
+
+		/// <summary>
+		/// The agent's working disk free space in MB.
+		/// </summary>
+		[Column("working_free_disk_mb")]
+		public long WorkingFreeDiskMb { get; init; }
+
+		/// <summary>
+		/// The agent's working disk total space in MB.
+		/// </summary>
+		[Column("working_total_disk_mb")]
+		public long WorkingTotalDiskMb { get; init; }
+
+		/// <summary>
+		/// The system's free disk in MB.
+		/// </summary>
+		[Column("system_free_disk_mb")]
+		public long SystemFreeDiskMb { get; init; }
+
+		/// <summary>
+		/// The system's total disk in MB.
+		/// </summary>
+		[Column("system_total_disk_mb")]
+		public long SystemTotalDiskMb { get; init; }
+
+		/// <summary>
+		/// Default event name for the Agent Telemetry event.
+		/// </summary>
+		public const string DefaultEventName = "State.Agent.Telemetry";
+
+		/// <summary>
+		/// Current schema version.
+		/// </summary>
+		public static readonly int CurrentSchemaVersion = 1;
+
+		/// <summary>
+		/// Constructor from pre-existing NewAgentTelemetry object
+		/// </summary>
+		/// <param name="telemetry">The telemetry object used as a basis for this telemetry record</param>
+		/// <param name="agentId">The agent's ID</param>
+		public AgentTelemetry(NewAgentTelemetry telemetry, AgentId agentId)
+			: base(DefaultEventName, CurrentSchemaVersion)
+		{
+			UserCpu = telemetry.UserCpuPct;
+			IdleCpu = telemetry.IdleCpuPct;
+			SystemCpu = telemetry.SystemCpuPct;
+			FreeRamMb = telemetry.FreeRamMb;
+			UsedRamMb = telemetry.UsedRamMb;
+			TotalRamMb = telemetry.TotalRamMb;
+			WorkingFreeDiskMb = telemetry.WorkingFreeDiskMb;
+			WorkingTotalDiskMb = telemetry.WorkingTotalDiskMb;
+			SystemFreeDiskMb = telemetry.SystemFreeDiskMb;
+			SystemTotalDiskMb = telemetry.SystemTotalDiskMb;
+
+			AgentId = agentId;
+		}
+
+		/// <summary>
+		/// Default constructor
+		/// </summary>
+		public AgentTelemetry()
+			: base(DefaultEventName, CurrentSchemaVersion)
+		{
+			UserCpu = 0;
+			IdleCpu = 0;
+			SystemCpu = 0;
+			FreeRamMb = 0;
+			UsedRamMb = 0;
+			TotalRamMb = 0;
+			WorkingFreeDiskMb = 0;
+			WorkingTotalDiskMb = 0;
+			SystemFreeDiskMb = 0;
+			SystemTotalDiskMb = 0;
+
+			AgentId = new AgentId(String.Empty);
+		}
+	}
+}
\ No newline at end of file
diff --git a/Engine/Source/Programs/Horde/Plugins/Compute/HordeServer.Compute/Server/RpcService.cs b/Engine/Source/Programs/Horde/Plugins/Compute/HordeServer.Compute/Server/RpcService.cs
--- a/Engine/Source/Programs/Horde/Plugins/Compute/HordeServer.Compute/Server/RpcService.cs
+++ b/Engine/Source/Programs/Horde/Plugins/Compute/HordeServer.Compute/Server/RpcService.cs
@@ -6,6 +6,7 @@ using System.Security.Claims;
 using EpicGames.Horde.Acls;
 using EpicGames.Horde.Agents;
 using EpicGames.Horde.Agents.Sessions;
+using EpicGames.Horde.Telemetry;
 using EpicGames.Horde.Tools;
 using Google.Protobuf;
 using Google.Protobuf.WellKnownTypes;
@@ -16,6 +17,7 @@ using HordeServer.Acls;
 using HordeServer.Agents;
 using HordeServer.Agents.Sessions;
 using HordeServer.Agents.Telemetry;
+using HordeServer.Telemetry;
 using HordeServer.Tools;
 using HordeServer.Utilities;
 using Microsoft.AspNetCore.Authorization;
@@ -42,11 +44,20 @@ namespace HordeServer.Server
 		readonly IAclService _aclService;
 		readonly IOptionsSnapshot<ComputeConfig> _computeConfig;
 		readonly ILogger _logger;
+		readonly ITelemetryWriter _telemetryWriter;
 
 		/// <summary>
 		/// Constructor
 		/// </summary>
-		public RpcService(AgentService agentService, ILifetimeService lifetimeService, IToolCollection toolCollection, IAgentTelemetryCollection agentTelemetryCollection, IAclService aclService, IOptionsSnapshot<ComputeConfig> computeConfig, ILogger<RpcService> logger)
+		public RpcService(
+			AgentService agentService, 
+			ILifetimeService lifetimeService, 
+			IToolCollection toolCollection, 
+			IAgentTelemetryCollection agentTelemetryCollection, 
+			IAclService aclService, 
+			IOptionsSnapshot<ComputeConfig> computeConfig, 
+			ILogger<RpcService> logger,
+			ITelemetryWriter telemetryWriter)
 		{
 			_agentService = agentService;
 			_lifetimeService = lifetimeService;
@@ -55,6 +66,7 @@ namespace HordeServer.Server
 			_aclService = aclService;
 			_computeConfig = computeConfig;
 			_logger = logger;
+			_telemetryWriter = telemetryWriter;
 		}
 
 		/// <summary>
@@ -421,14 +433,33 @@ namespace HordeServer.Server
 		/// <inheritdoc/>
 		public override Task<Empty> UploadTelemetry(RpcUploadTelemetryRequest request, ServerCallContext context)
 		{
-			_logger.LogDebug("Posting telemetry data for {AgentId}", new AgentId(request.AgentId));
+			AgentId agentId = new AgentId(request.AgentId);
+
+			_logger.LogDebug("Posting telemetry data for {AgentId}", agentId);
 
 			NewAgentTelemetry telemetry = new NewAgentTelemetry(
 				request.UserCpu, request.IdleCpu, request.SystemCpu, 
 				(int)request.FreeRam, (int)request.UsedRam, (int)request.TotalRam, 
 				(long)request.WorkingFreeDisk, (long)request.WorkingTotalDisk,
 				(long)request.SystemFreeDisk, (long)request.SystemTotalDisk);
-			_agentTelemetryCollection.Add(new AgentId(request.AgentId), telemetry);
+			_agentTelemetryCollection.Add(agentId, telemetry);
+
+			if (_telemetryWriter.Enabled)
+			{
+				_logger.LogDebug("Writing telemetry event for {AgentId} to data sink(s).", agentId);
+
+				try
+				{
+					_telemetryWriter.WriteEvent(
+						TelemetryStoreId.Default,
+						new AgentTelemetry(telemetry, agentId)
+					);
+				}
+				catch (Exception ex)
+				{
+					_logger.LogError(ex, "Failed to publish agent telemetry.");
+				}
+			}
 
 			return Task.FromResult(new Empty());
 		}
