Horde: Improve checkbox visibility in dark mode
Brighten enabled checkbox border and add brighter hover state for better contrast.
Add explicit disabled state border to distinguish from enabled.
Fix label hover to match checkbox box hover behavior, preserving existing Fluent checked+hover state with :not(.is-checked) pattern.

#Synced-CL 0

diff --git a/Engine/Source/Programs/Horde/HordeDashboard/src/styles/darkTheme.ts b/Engine/Source/Programs/Horde/HordeDashboard/src/styles/darkTheme.ts
--- a/Engine/Source/Programs/Horde/HordeDashboard/src/styles/darkTheme.ts
+++ b/Engine/Source/Programs/Horde/HordeDashboard/src/styles/darkTheme.ts
@@ -28,6 +28,11 @@ const basePrimary = "#0078D4"
 const borderColor = neutralLight
 const borderHoverColor = "#4f4f4f"
 
+// Checkbox border colors (brighter for visibility in dark mode)
+const checkboxBorderColor = "#6F6F6F"
+const checkboxBorderHoverColor = "#8F8F8F"
+const checkboxBorderDisabledColor = "#4F4F4F"
+
 // Sub-menu divider color
 const menuHeader = "#2b94fe";
 
@@ -251,13 +256,20 @@ export const darkTheme = createTheme({
         "Checkbox": {
             styles: {
                 checkmark: {
-                    color: "#FFFFFF"                    
+                    color: "#FFFFFF"
                 },
                 checkbox: {
-                    ":hover": {
-                        borderColor: `${borderHoverColor} !important`,                       
-                    },                    
-                    borderColor: borderColor,
+                    borderColor: checkboxBorderColor,
+                },
+                root: {
+                    selectors: {
+                        "&:not(.is-checked):hover .ms-Checkbox-checkbox": {
+                            borderColor: `${checkboxBorderHoverColor} !important`,
+                        },
+                        "&.is-disabled .ms-Checkbox-checkbox": {
+                            borderColor: `${checkboxBorderDisabledColor} !important`,
+                        }
+                    }
                 }
             }
         },
