Horde: Update release notes

#rnx

diff --git a/Engine/Source/Programs/Horde/Docs/ReleaseNotes.md b/Engine/Source/Programs/Horde/Docs/ReleaseNotes.md
--- a/Engine/Source/Programs/Horde/Docs/ReleaseNotes.md
+++ b/Engine/Source/Programs/Horde/Docs/ReleaseNotes.md
@@ -1,5 +1,33 @@
 # Release Notes
 
+## 2026-02-05
+
+* Improve checkbox visibility in dark mode
+  Brighten enabled checkbox border and add brighter hover state for better contrast.
+  Add explicit disabled state border to distinguish from enabled.
+  Fix label hover to match checkbox box hover behavior, preserving existing Fluent checked+hover state with :not(.is-checked) pattern. (50525707)
+* [Horde Server][Analytics] Propagating full agent telemetry events to associated data sinks (50509499)
+* Refactor SlackClient error handling to use enum-based behavior instead of boolean predicates
+  * Logs "not_in_channel" error message at log level warning instead of error. This error can be common if the Horde Slack bot has not been invited to a channel.
+  * For throw error behavior, it's assumed that the caller will take care of logging any exceptions (previously the exception was logged, and then re-thrown) (50483048)
+* [Horde Server][Horde Dashboard][Analytics] Adding system drive stats to CapabilitiesService + displaying new drive stats on Horde Dashboard
+  - Added new set of system drive related disk info events propagated through the CapabilitiesService (`C:\` on Windows and `/` on Unix). 
+  - Added HordeDashboard boilerplate to support these new analytics events, and modified nomenclature of pre-existing "working drive" telemetry (50452739)
+* Remove full URL from OpenTelemetry trace name in Horde Agent
+  * It has too high cardinality when IDs, CLs etc get included
+  * Added support for passing along a request option called "url.template" to let calling code provide a better name hint (50436066)
+* Fix build health compiler warnings (50435967)
+* Add the option to clean into RetrieveCloudArtifact so we can clean the folder for test jobs. (50403203)
+* Dispose Perforce PoolConnectionHandle when there is an exception during login (50395731)
+* Add optional helpSlack/helpEmail overrides to ErrorInfo for per-error support contact customization (50381085)
+* [Horde Server][Horde Agent][Analytics] Propagation of Agent's system drive usage stats to Agent telemetry reporting
+  Modifies the original set of system metrics propagated into mongodb to split the concept of agent disk (what we're currently tracking) and system disk (C:\ on Windows, '/' on Unix. This data will eventually be consumed by the Horde telemetry dashboard to show a 2nd graph tracking system disk size alongside the pre-existing agent disk size graph. (50354394)
+* Make upgrade of Windows agent more robust with retry logic
+  Add retry for file deletion and move operations during agent self-upgrade to handle lingering file locks after service stop. (50350263)
+* Ensure we use the targeted framework (.net 8)
+  By default we use LatestMajor which may not be desired. (50316173)
+* Missed file for agent telemetry (50306429)
+
 ## 2026-01-29
 
 * Add error handling for stuck step timeout check in JobExpirationService (50303167)
