Horde: Fix disbaled checkbox hover

#rnx

#Synced-CL 0

diff --git a/Engine/Source/Programs/Horde/HordeDashboard/src/styles/darkTheme.ts b/Engine/Source/Programs/Horde/HordeDashboard/src/styles/darkTheme.ts
--- a/Engine/Source/Programs/Horde/HordeDashboard/src/styles/darkTheme.ts
+++ b/Engine/Source/Programs/Horde/HordeDashboard/src/styles/darkTheme.ts
@@ -263,7 +263,7 @@ export const darkTheme = createTheme({
                 },
                 root: {
                     selectors: {
-                        "&:not(.is-checked):hover .ms-Checkbox-checkbox": {
+                        "&:not(.is-checked):not(.is-disabled):hover .ms-Checkbox-checkbox": {
                             borderColor: `${checkboxBorderHoverColor} !important`,
                         },
                         "&.is-disabled .ms-Checkbox-checkbox": {
