[Horde Dashboard] Fixing improper property names for fetching storage metrics within `Agents` view

#tests tested with a local built version of the horde dashboard targeting the horde prod environment

diff --git a/Engine/Source/Programs/Horde/HordeDashboard/src/components/agents/AgentTelemetrySparkline.tsx b/Engine/Source/Programs/Horde/HordeDashboard/src/components/agents/AgentTelemetrySparkline.tsx
--- a/Engine/Source/Programs/Horde/HordeDashboard/src/components/agents/AgentTelemetrySparkline.tsx
+++ b/Engine/Source/Programs/Horde/HordeDashboard/src/components/agents/AgentTelemetrySparkline.tsx
@@ -208,7 +208,7 @@ export const AgentTelemetrySparkline: React.FC<{ handler: AgentTelemetryHandler
          ramText = `${v.replace("RAM=", "")} GB`;
       }
 
-      if (v.startsWith("DiskTotalSize=")) {
+      if (v.startsWith("WorkingDiskTotalSize=")) {
          const elements = v.split("=");
          if (elements.length === 2) {
 
diff --git a/Engine/Source/Programs/Horde/HordeDashboard/src/components/agents/AgentView.tsx b/Engine/Source/Programs/Horde/HordeDashboard/src/components/agents/AgentView.tsx
--- a/Engine/Source/Programs/Horde/HordeDashboard/src/components/agents/AgentView.tsx
+++ b/Engine/Source/Programs/Horde/HordeDashboard/src/components/agents/AgentView.tsx
@@ -2193,8 +2193,8 @@ const AgentViewInner: React.FC<{ agentId?: string, poolId?: string, searchParams
                   toCheck = localState.columnSearchState['storage'][agent.id];
                }
                else {
-                  let diskFree: string | number | null = getAgentCapability(agent, "DiskFreeSpace");
-                  let diskCapacity: string | number | null = getAgentCapability(agent, "DiskTotalSize");
+                  let diskFree: string | number | null = getAgentCapability(agent, "WorkingDiskFreeSpace");
+                  let diskCapacity: string | number | null = getAgentCapability(agent, "WorkingDiskTotalSize");
                   toCheck = "Unknown";
                   if (diskFree !== null && diskCapacity !== null) {
                      diskFree = Number(diskFree) / 1073741824;
@@ -2876,8 +2876,8 @@ const AgentViewInner: React.FC<{ agentId?: string, poolId?: string, searchParams
                </Stack.Item>
             </Stack>);
          case 'storage':
-            let diskFree: string | number | null = getAgentCapability(agent, "DiskFreeSpace");
-            let diskCapacity: string | number | null = getAgentCapability(agent, "DiskTotalSize");
+            let diskFree: string | number | null = getAgentCapability(agent, "WorkingDiskFreeSpace");
+            let diskCapacity: string | number | null = getAgentCapability(agent, "WorkingDiskTotalSize");
 
             let realPercentage = 0;
             let nudgedPercentage = 0;
