[Horde][Dashboard] Backend library now supports cancellations.

#Synced-CL 50425380

diff --git a/Engine/Source/Programs/Horde/HordeDashboard/src/backend/Fetch.ts b/Engine/Source/Programs/Horde/HordeDashboard/src/backend/Fetch.ts
--- a/Engine/Source/Programs/Horde/HordeDashboard/src/backend/Fetch.ts
+++ b/Engine/Source/Programs/Horde/HordeDashboard/src/backend/Fetch.ts
@@ -19,6 +19,7 @@ export interface FetchRequestConfig {
     suppress404?: boolean;
     // shows error dialog for 404, in addition to logging it
     show404Error?: boolean;
+    signal?: AbortSignal;
 }
 
 export type FetchResponse = {
@@ -27,7 +28,7 @@ export type FetchResponse = {
 }
 
 function handleError(
-    errorIn: ErrorInfo, 
+    errorIn: ErrorInfo,
     force?: boolean,
     reportType?: ErrorReportType,
     retryCallback?: () => boolean | Promise<boolean>
@@ -35,7 +36,7 @@ function handleError(
 
     const e = { ...errorIn, hordeError: true } as ErrorInfo;
 
-    switch(reportType){
+    switch (reportType) {
         case ErrorReportType.Dialog:
             errorDialogStore.set(e, !!force, retryCallback);
             break;
@@ -55,15 +56,15 @@ export class Fetch {
 
         url = this.buildUrl(url, config);
 
-        if(this.checkDebugMockStatus(url)){
+        if (this.checkDebugMockStatus(url)) {
             return Promise.reject(new Error("Vite endpoint status debugging"))
         }
 
         return new Promise<FetchResponse>(async (resolve, reject) => {
 
             try {
-                
-                const response = await fetch(url, this.buildRequest("GET")).then(response => {
+
+                const response = await fetch(url, this.buildRequest("GET", undefined, config)).then(response => {
                     // Internal Server Error Response
                     if (!response.ok && response.status === 500) {
                         response.text().then(text => {
@@ -73,7 +74,7 @@ export class Fetch {
                                 mode: "GET",
                                 url: url,
                                 message: text
-                            }, true, ErrorReportType.Dialog);            
+                            }, true, ErrorReportType.Dialog);
                         });
                         return null;
                     }
@@ -88,7 +89,7 @@ export class Fetch {
                                 mode: "GET",
                                 url: url,
                                 message: o?.message ?? "Malformed json on response object"
-                            }, true, ErrorReportType.Dialog);            
+                            }, true, ErrorReportType.Dialog);
                         }).catch(reason => {
                             console.error("Unable to parse response json on 404: ", reason)
                         });
@@ -96,7 +97,7 @@ export class Fetch {
                     }
 
                     // 401 Unauthorized Response - No Credentials
-                    if(!response.ok && response.status === 401){
+                    if (!response.ok && response.status === 401) {
                         handleError({
                             response,
                             title: "401 Unauthorized",
@@ -109,7 +110,7 @@ export class Fetch {
                     }
 
                     // 403 Forbidden Response - Insufficient Permissions
-                    if(!response.ok && response.status === 403){
+                    if (!response.ok && response.status === 403) {
                         handleError({
                             response,
                             title: "403 Forbidden",
@@ -128,7 +129,7 @@ export class Fetch {
                     reject('');
                     return;
                 }
-		
+
                 if (response.ok) {
                     this.handleResponse(response, url, "GET", resolve, reject, config);
                     return;
@@ -174,11 +175,11 @@ export class Fetch {
 
     // DEBUG code, used to test response status handling.
     checkDebugMockStatus(url: string): boolean {
-        if(!import.meta.env.VITE_HORDE_DEBUG_MOCK_STATUS_ENDPOINTS){
+        if (!import.meta.env.VITE_HORDE_DEBUG_MOCK_STATUS_ENDPOINTS) {
             // No debug settings found, return false to continue the calling function
             return false;
         }
-        
+
         const mockStatusEndpoints: Record<string, [number, ErrorReportType]> = {};
         // Parse the vite mock status envrionment variable into a record map-like object
         import.meta.env.VITE_HORDE_DEBUG_MOCK_STATUS_ENDPOINTS
@@ -186,18 +187,18 @@ export class Fetch {
             .forEach((triplet, index) => {
                 try {
                     const [endpointRaw, statusRaw, reportTypeRaw] = triplet.split(":");
-                    
+
                     // Endpoint parsing
-                    if(!endpointRaw){
+                    if (!endpointRaw) {
                         throw new Error(`Missing endpoint in triplet "${triplet}" at index ${index}`);
                     }
                     const endpoint = endpointRaw.trim();
-                    
+
                     // Status parsing
                     let status = 401;
-                    if(statusRaw){
+                    if (statusRaw) {
                         const parsedStatus = Number(statusRaw.trim());
-                        if(isNaN(parsedStatus) || parsedStatus < 100 || parsedStatus > 599) {
+                        if (isNaN(parsedStatus) || parsedStatus < 100 || parsedStatus > 599) {
                             console.warn(`Invalid status "${statusRaw}" in triplet "${triplet}" at index ${index}. Defaulting to status 401.`);
                         }
                         else {
@@ -210,11 +211,11 @@ export class Fetch {
 
                     // Report type parsing
                     let reportType = ErrorReportType.Bar;
-                    if(reportTypeRaw){
+                    if (reportTypeRaw) {
                         const parsedReportType = ErrorReportType[reportTypeRaw.trim() as keyof ErrorReportType];
-                        if(!parsedReportType){
+                        if (!parsedReportType) {
                             console.warn(`Invalid report type "${reportTypeRaw}" in triplet "${triplet}" at index ${index}. Defaulting to ErrorReportType "Bar".`);
-                        } 
+                        }
                         else {
                             reportType = parsedReportType;
                         }
@@ -232,12 +233,12 @@ export class Fetch {
         // Determine if any entries match the current url
         const match = Object.entries(mockStatusEndpoints).find(([endpoint]) => url.includes(endpoint));
         // If a match is found, simulate a bad response with the given status and mock an error with the given report type
-        if(match) {
+        if (match) {
             const [endpoint, [status, reportType]] = match;
 
-            const response = new Response(JSON.stringify({error: `Error ${status} (Mock)`}), {
+            const response = new Response(JSON.stringify({ error: `Error ${status} (Mock)` }), {
                 status,
-                headers: {"Content-Type": "application/json"}
+                headers: { "Content-Type": "application/json" }
             });
 
             if (!response.ok) {
@@ -292,7 +293,7 @@ export class Fetch {
 
         return new Promise<FetchResponse>((resolve, reject) => {
 
-            fetch(url, this.buildRequest("PUT", data)).then(response => {
+            fetch(url, this.buildRequest("PUT", data, config)).then(response => {
 
                 this.handleResponse(response, url, "PUT", resolve, reject, config);
 
@@ -316,7 +317,7 @@ export class Fetch {
 
         return new Promise<FetchResponse>((resolve, reject) => {
 
-            fetch(url, this.buildRequest("PATCH", data)).then(response => {
+            fetch(url, this.buildRequest("PATCH", data, config)).then(response => {
 
                 this.handleResponse(response, url, "PATCH", resolve, reject, config);
 
@@ -341,7 +342,7 @@ export class Fetch {
 
         return new Promise<FetchResponse>((resolve, reject) => {
 
-            fetch(url, this.buildRequest("DELETE")).then(response => {
+            fetch(url, this.buildRequest("DELETE", undefined, config)).then(response => {
 
                 this.handleResponse(response, url, "DELETE", resolve, reject, config);
 
@@ -377,7 +378,7 @@ export class Fetch {
                     response: response,
                     url: url,
                     title: "Internal Server Error",
-                    message: text                
+                    message: text
                 }, true, ErrorReportType.Dialog);
             });
             return reject("Internal Server Error");
@@ -509,7 +510,8 @@ export class Fetch {
             credentials: this.debugToken ? "same-origin" : 'include',
             redirect: "error",
             headers: headers,
-            body: body
+            body: body,
+            signal: config?.signal
         };
 
         return requestInit;
@@ -526,7 +528,7 @@ export class Fetch {
         if (this.withCredentials) {
             headers["Authorization"] = this.authorization!;
         }
-        
+
         // Report how idle the user has been on current page in each HTTP request
         // This helps identify how much load potentially unused pages generate on the Horde server
         headers["X-Horde-LastUserActivity"] = userInactivity.getSecondsSinceLastActivity().toString();
@@ -581,8 +583,8 @@ export class Fetch {
             const url = this.buildUrl("/api/v1/dashboard/challenge");
 
             const result = await fetch(url, this.buildRequest("GET"));
-            const response = await result.json() as GetDashboardChallengeResponse;      
-            
+            const response = await result.json() as GetDashboardChallengeResponse;
+
             if (response.needsFirstTimeSetup) {
                 window.location.assign("/setup");
                 // give the window assignment a couple seconds, so we don't continue down challenge route
diff --git a/Engine/Source/Programs/Horde/HordeDashboard/src/backend/RequestCancellationManager.ts b/Engine/Source/Programs/Horde/HordeDashboard/src/backend/RequestCancellationManager.ts
new file mode 100644
--- /dev/null
+++ b/Engine/Source/Programs/Horde/HordeDashboard/src/backend/RequestCancellationManager.ts
@@ -0,0 +1,49 @@
+
+// Copyright Epic Games, Inc. All Rights Reserved.
+
+/**
+ * Cancellation reason when a subsequent request to the same @see AbortController has occurred.
+ */
+export const REASON_REENTRANT_REQUEST = "REENTRANT_REQUEST";
+
+/**
+ * Lightweight class to help with managing cancellations for various endpoints.
+ */
+export class RequestCancellationManager {
+
+    private controllers: Map<string, AbortController> = new Map<string, AbortController>();
+
+    /**
+     * Obtains a new cancellation signal for API requests, cancelling any before that with the @see REASON_REENTRANT_REQUEST reason.
+     * @param identifier The identifier associated with the @see AbortController
+     * @returns The signal of the active abort controller assigned to the identifier.
+     */
+    getSignal(identifier: string): AbortSignal {
+        this.controllers.get(identifier)?.abort(REASON_REENTRANT_REQUEST);
+        let newController: AbortController = new AbortController();
+        this.controllers.set(identifier, newController);
+
+        return newController.signal;
+    }
+
+    /**
+     * Cancels an existing @see AbortController given the identifier, should it exist.
+     * @param identifier The identifier associated with the @see AbortController
+     * @param reason The reason to issue the cancel.
+     */
+    cancel(identifier: string, reason?: string): void {
+        this.controllers.get(identifier)?.abort(reason);
+        this.controllers.delete(identifier);
+    }
+
+    /**
+     * Cancels all managed @see AbortController .
+     * @param reason The reason to issue a cancel all.
+     */
+    cancelAll(reason?: string): void {
+        this.controllers.forEach((value: AbortController) => {
+            value.abort(reason);
+        });
+        this.controllers.clear();
+    }
+}
\ No newline at end of file
