Horde: Fix gRPC streaming errors in UpdateLogTail

The UpdateLogTail bidirectional streaming RPC was throwing unhandled exceptions under certain conditions when clients closed connections (timeout, process termination, or intentional closure).

diff --git a/Engine/Source/Programs/Horde/Plugins/Compute/HordeServer.Compute/Logs/LogRpcService.cs b/Engine/Source/Programs/Horde/Plugins/Compute/HordeServer.Compute/Logs/LogRpcService.cs
--- a/Engine/Source/Programs/Horde/Plugins/Compute/HordeServer.Compute/Logs/LogRpcService.cs
+++ b/Engine/Source/Programs/Horde/Plugins/Compute/HordeServer.Compute/Logs/LogRpcService.cs
@@ -20,6 +20,12 @@ namespace HordeServer.Logs
 	[Authorize]
 	public class LogRpcService : LogRpc.LogRpcBase
 	{
+		// Exception message substrings for detecting closed gRPC response streams.
+		// These come from Grpc.AspNetCore.Server (HttpContextServerCallContext) and ASP.NET Core's response body writer.
+		// .NET Core/.NET 5+ exception messages are not localized (always English), so string matching is safe.
+		const string RequestCompleteMessage = "Request is complete";
+		const string WritingNotAllowedMessage = "Writing is not allowed";
+
 		readonly ILogCollection _logCollection;
 		readonly LogTailService _logTailService;
 		readonly StorageService _storageService;
@@ -102,24 +108,58 @@ namespace HordeServer.Logs
 					try
 					{
 						response.TailNext = await waitTask;
-						await completeTask; // Allow moveNextTask to throw if it finished due to cancellation, so we don't write to
+						await completeTask; // Allow moveNextTask to throw if it finished due to cancellation, so we don't write to broken stream
 					}
 					catch (OperationCanceledException)
 					{
 						_logger.LogDebug("Log tail stream cancelled by client");
 						break;
 					}
+					catch (IOException ex)
+					{
+						// Connection was aborted (client deadline exceeded, client process killed, etc.)
+						_logger.LogDebug(ex, "Log tail stream connection aborted for {LogId}", logId);
+						break;
+					}
 					catch (Exception ex)
 					{
 						_logger.LogWarning(ex, "Exception while waiting for tail next");
+						break; // Don't continue on unknown errors - stream may be broken
 					}
 				}
 
 				_logger.LogDebug("Return tail next for log {LogId} = {TailNext}", logId, response.TailNext);
-				await responseStream.WriteAsync(response);
+
+				// Guard the write against connection closure
+				if (!await TryWriteResponseAsync(responseStream, response))
+				{
+					break;
+				}
 			}
 
-			await responseStream.WriteAsync(response);
+			// Don't write after loop - stream is closed when we exit the while loop
+		}
+
+		/// <summary>
+		/// Attempts to write a response to the stream, returning false if the stream is already closed
+		/// </summary>
+		static async Task<bool> TryWriteResponseAsync(IServerStreamWriter<RpcUpdateLogTailResponse> responseStream, RpcUpdateLogTailResponse response)
+		{
+			try
+			{
+				await responseStream.WriteAsync(response);
+				return true;
+			}
+			catch (InvalidOperationException ex) when (ex.Message.Contains(RequestCompleteMessage, StringComparison.OrdinalIgnoreCase) || ex.Message.Contains(WritingNotAllowedMessage, StringComparison.OrdinalIgnoreCase))
+			{
+				// Stream already closed by client - expected during shutdown
+				return false;
+			}
+			catch (IOException)
+			{
+				// Connection aborted
+				return false;
+			}
 		}
 
 		/// <inheritdoc/>
