Horde: Add material-based grouping to ShaderIssueHandler
This ensures all failures for the same material are tracked together.
Fix LogShaderCompilers grouping by expanding ShaderEventMatcher to capture continuation lines and extract material paths.

#Synced-CL 0

diff --git a/Engine/Source/Programs/Horde/Plugins/Build/HordeServer.Build.Tests/Issues/IssueHandlerTests.cs b/Engine/Source/Programs/Horde/Plugins/Build/HordeServer.Build.Tests/Issues/IssueHandlerTests.cs
--- a/Engine/Source/Programs/Horde/Plugins/Build/HordeServer.Build.Tests/Issues/IssueHandlerTests.cs
+++ b/Engine/Source/Programs/Horde/Plugins/Build/HordeServer.Build.Tests/Issues/IssueHandlerTests.cs
@@ -187,5 +187,75 @@ namespace HordeServer.Tests.Issues
 			Assert.IsTrue(groupedIssues[8].SummaryTemplate == "{Meta:Node}: Localization {Severity} in {Meta:Stream}");
 			Assert.IsTrue(groupedIssues[8].Type == "Localization");
 		}
+
+		[TestMethod]
+		public void ShaderIssueHandler_GroupsByMaterialWhenPropertyPresent()
+		{
+			// These are shader warnings for the same material failing to compile on different platforms/shaders.
+			// When the "material" property is present (set by ShaderEventMatcher), all warnings for the same
+			// material should be keyed by the material path, enabling them to be grouped together.
+			string[] jsonStructuredLogs =
+			{
+				// Shader warning with material property - will be keyed by material path
+				"{ \"time\":\"2025-01-01T01:20:30.990Z\",\"level\":\"Warning\",\"message\":\"LogShaderCompilers: Warning: Failed to compile Material /MeshModelingToolsetExp/Materials/PointSetOverlaidComponentMaterial.PointSetOverlaidComponentMaterial for platform SF_PS5\",\"format\":\"{_channel}: {_severity}: {message}\",\"properties\":{ \"_channel\":{ \"$type\":\"Channel\",\"$text\":\"LogShaderCompilers\"},\"_severity\":{ \"$type\":\"Severity\",\"$text\":\"Warning\"},\"id\":310,\"file\":\"D:/build/Engine/Shaders/Private/LocalVertexFactory.ush\",\"material\":\"/MeshModelingToolsetExp/Materials/PointSetOverlaidComponentMaterial.PointSetOverlaidComponentMaterial\"} }",
+				// Same material, different platform - should have same key
+				"{ \"time\":\"2025-01-01T01:20:30.991Z\",\"level\":\"Warning\",\"message\":\"LogShaderCompilers: Warning: Failed to compile Material /MeshModelingToolsetExp/Materials/PointSetOverlaidComponentMaterial.PointSetOverlaidComponentMaterial for platform SF_VULKAN_SM6\",\"format\":\"{_channel}: {_severity}: {message}\",\"properties\":{ \"_channel\":{ \"$type\":\"Channel\",\"$text\":\"LogShaderCompilers\"},\"_severity\":{ \"$type\":\"Severity\",\"$text\":\"Warning\"},\"id\":310,\"file\":\"D:/build/Engine/Shaders/Private/MaterialTemplate.ush\",\"material\":\"/MeshModelingToolsetExp/Materials/PointSetOverlaidComponentMaterial.PointSetOverlaidComponentMaterial\"} }",
+				// Different material - should have different key
+				"{ \"time\":\"2025-01-01T01:20:30.992Z\",\"level\":\"Warning\",\"message\":\"LogShaderCompilers: Warning: Failed to compile Material /Engine/EngineMaterials/LineSetComponentMaterial.LineSetComponentMaterial for platform PCD3D_SM6\",\"format\":\"{_channel}: {_severity}: {message}\",\"properties\":{ \"_channel\":{ \"$type\":\"Channel\",\"$text\":\"LogShaderCompilers\"},\"_severity\":{ \"$type\":\"Severity\",\"$text\":\"Warning\"},\"id\":310,\"file\":\"D:/build/Engine/Shaders/Private/LocalVertexFactory.ush\",\"material\":\"/Engine/EngineMaterials/LineSetComponentMaterial.LineSetComponentMaterial\"} }",
+			};
+
+			ShaderIssueHandler shaderIssueHandler = new ShaderIssueHandler();
+			List<IssueEventGroup> groupedIssues = TestSingleIssueHandler(shaderIssueHandler, jsonStructuredLogs);
+
+			// ShaderIssueHandler creates separate issues but keys them by material path when property is present
+			Assert.AreEqual(3, groupedIssues.Count);
+
+			// First two issues have the same material key (PointSetOverlaidComponentMaterial)
+			Assert.AreEqual(1, groupedIssues[0].Keys.Count);
+			Assert.IsTrue(groupedIssues[0].Keys.Contains(new IssueKey("/MeshModelingToolsetExp/Materials/PointSetOverlaidComponentMaterial.PointSetOverlaidComponentMaterial", IssueKeyType.File)));
+			Assert.AreEqual("Shader", groupedIssues[0].Type);
+			Assert.IsTrue(groupedIssues[0].Metadata.FindValues("Material").Any());
+			Assert.AreEqual("PointSetOverlaidComponentMaterial", groupedIssues[0].Metadata.FindValues("Material").First());
+
+			Assert.AreEqual(1, groupedIssues[1].Keys.Count);
+			Assert.IsTrue(groupedIssues[1].Keys.Contains(new IssueKey("/MeshModelingToolsetExp/Materials/PointSetOverlaidComponentMaterial.PointSetOverlaidComponentMaterial", IssueKeyType.File)));
+
+			// Third issue has different material key (LineSetComponentMaterial)
+			Assert.AreEqual(1, groupedIssues[2].Keys.Count);
+			Assert.IsTrue(groupedIssues[2].Keys.Contains(new IssueKey("/Engine/EngineMaterials/LineSetComponentMaterial.LineSetComponentMaterial", IssueKeyType.File)));
+			Assert.AreEqual("LineSetComponentMaterial", groupedIssues[2].Metadata.FindValues("Material").First());
+		}
+
+		[TestMethod]
+		public void ShaderIssueHandler_GroupsBySourceFileWhenNoMaterialProperty()
+		{
+			// These are shader errors WITHOUT the material property (e.g., shader syntax errors).
+			// They should continue to use the existing behavior: grouped by source file.
+			string[] jsonStructuredLogs =
+			{
+				// Shader error with source file only (no material property)
+				"{ \"time\":\"2025-01-01T01:20:30.990Z\",\"level\":\"Error\",\"message\":\"LogShaderCompilers: Error: Shader compilation failed in LocalVertexFactory.ush\",\"format\":\"{_channel}: {_severity}: {message}\",\"properties\":{ \"_channel\":{ \"$type\":\"Channel\",\"$text\":\"LogShaderCompilers\"},\"_severity\":{ \"$type\":\"Severity\",\"$text\":\"Error\"},\"id\":310,\"file\":\"D:/build/Engine/Shaders/Private/LocalVertexFactory.ush\"} }",
+				// Different source file
+				"{ \"time\":\"2025-01-01T01:20:30.991Z\",\"level\":\"Error\",\"message\":\"LogShaderCompilers: Error: Shader compilation failed in MaterialTemplate.ush\",\"format\":\"{_channel}: {_severity}: {message}\",\"properties\":{ \"_channel\":{ \"$type\":\"Channel\",\"$text\":\"LogShaderCompilers\"},\"_severity\":{ \"$type\":\"Severity\",\"$text\":\"Error\"},\"id\":310,\"file\":\"D:/build/Engine/Shaders/Private/MaterialTemplate.ush\"} }",
+			};
+
+			ShaderIssueHandler shaderIssueHandler = new ShaderIssueHandler();
+			List<IssueEventGroup> groupedIssues = TestSingleIssueHandler(shaderIssueHandler, jsonStructuredLogs);
+
+			// Without material property, issues are keyed by source file
+			Assert.AreEqual(2, groupedIssues.Count);
+
+			// First issue keyed by LocalVertexFactory.ush
+			Assert.AreEqual(1, groupedIssues[0].Keys.Count);
+			Assert.IsTrue(groupedIssues[0].Keys.Contains(new IssueKey("LocalVertexFactory.ush", IssueKeyType.File)));
+			Assert.AreEqual("Shader", groupedIssues[0].Type);
+			Assert.IsFalse(groupedIssues[0].Metadata.FindValues("Material").Any());
+
+			// Second issue keyed by MaterialTemplate.ush
+			Assert.AreEqual(1, groupedIssues[1].Keys.Count);
+			Assert.IsTrue(groupedIssues[1].Keys.Contains(new IssueKey("MaterialTemplate.ush", IssueKeyType.File)));
+			Assert.IsFalse(groupedIssues[1].Metadata.FindValues("Material").Any());
+		}
+
 	}
 }
diff --git a/Engine/Source/Programs/Shared/EpicGames.Horde/Issues/Handlers/ShaderIssueHandler.cs b/Engine/Source/Programs/Shared/EpicGames.Horde/Issues/Handlers/ShaderIssueHandler.cs
--- a/Engine/Source/Programs/Shared/EpicGames.Horde/Issues/Handlers/ShaderIssueHandler.cs
+++ b/Engine/Source/Programs/Shared/EpicGames.Horde/Issues/Handlers/ShaderIssueHandler.cs
@@ -1,13 +1,21 @@
 // Copyright Epic Games, Inc. All Rights Reserved.
 
 using System.Collections.Generic;
+using System.Text.Json;
 using EpicGames.Core;
 using Microsoft.Extensions.Logging;
 
 namespace EpicGames.Horde.Issues.Handlers
 {
 	/// <summary>
-	/// Instance of a particular shader compile error
+	/// Handles shader compile errors and warnings, creating issue groups for tracking.
+	///
+	/// Grouping behavior:
+	/// - For "Failed to compile Material" warnings: Groups by material path, so all compilation
+	///   failures for the same material (across different platforms/shaders) are tracked as one issue.
+	/// - For other shader errors: Groups by source file, creating separate issues per file.
+	///
+	/// The material path is extracted by ShaderEventMatcher and passed as a "material" property.
 	/// </summary>
 	[IssueHandler]
 	public class ShaderIssueHandler : IssueHandler
@@ -28,11 +36,30 @@ namespace EpicGames.Horde.Issues.Handlers
 				EventId eventId = logEvent.EventId.Value;
 				if (IsMatchingEventId(eventId))
 				{
-					IssueEventGroup issue = new IssueEventGroup("Shader", "Shader compile {Severity} in {Files}", IssueChangeFilter.Code);
+					// Check if this event has a material property (set by ShaderEventMatcher for
+					// "Failed to compile Material" warnings). If so, group by material path instead
+					// of source files, so all failures for the same material are tracked together.
+					string? materialPath = TryGetMaterialPath(logEvent);
+
+					IssueEventGroup issue;
+					if (materialPath != null)
+					{
+						// Group by material path - use a different summary template that references the material
+						issue = new IssueEventGroup("Shader", "Shader compile {Severity} in {Meta:Material}", IssueChangeFilter.Code);
+						issue.Keys.Add(new IssueKey(materialPath, IssueKeyType.File));
+						issue.Metadata.Add("Material", GetMaterialDisplayName(materialPath));
+						logEvent.AuditLogger?.LogDebug("{IssueType} issue added for material '{Material}': '{Event}'", issue.Type, materialPath, logEvent.Render());
+					}
+					else
+					{
+						// Default behavior: group by source files
+						issue = new IssueEventGroup("Shader", "Shader compile {Severity} in {Files}", IssueChangeFilter.Code);
+						issue.Keys.AddSourceFiles(logEvent);
+						logEvent.AuditLogger?.LogDebug("{IssueType} issue added for event: '{Event}'", issue.Type, logEvent.Render());
+					}
+
 					issue.Events.Add(logEvent);
-					issue.Keys.AddSourceFiles(logEvent);
 					_issues.Add(issue);
-					logEvent.AuditLogger?.LogDebug("{IssueType} issue added for event: '{Event}'", issue.Type, logEvent.Render());
 
 					return true;
 				}
@@ -44,6 +71,53 @@ namespace EpicGames.Horde.Issues.Handlers
 			return false;
 		}
 
+		/// <summary>
+		/// Extracts the material path from the event's "material" property, if present.
+		/// This property is set by ShaderEventMatcher when it detects "Failed to compile Material" warnings.
+		/// </summary>
+		private static string? TryGetMaterialPath(IssueEvent logEvent)
+		{
+			foreach (JsonLogEvent line in logEvent.Lines)
+			{
+				try
+				{
+					using JsonDocument document = JsonDocument.Parse(line.Data);
+					if (document.RootElement.TryGetProperty("properties", out JsonElement properties) &&
+						properties.ValueKind == JsonValueKind.Object &&
+						properties.TryGetProperty("material", out JsonElement materialElement) &&
+						materialElement.ValueKind == JsonValueKind.String)
+					{
+						return materialElement.GetString();
+					}
+				}
+				catch (JsonException)
+				{
+					// Ignore parsing errors and continue checking other lines
+				}
+			}
+			return null;
+		}
+
+		/// <summary>
+		/// Extracts a display-friendly material name from the full material path.
+		/// For example: "/Game/Materials/MyMaterial.MyMaterial" -> "MyMaterial"
+		/// </summary>
+		private static string GetMaterialDisplayName(string materialPath)
+		{
+			// Material paths are typically "/Package/Path/MaterialName.MaterialName"
+			// Extract just the material name (last component after the final slash, before the dot)
+			int lastSlash = materialPath.LastIndexOf('/');
+			string name = lastSlash >= 0 ? materialPath.Substring(lastSlash + 1) : materialPath;
+
+			int dotIndex = name.IndexOf('.');
+			if (dotIndex > 0)
+			{
+				name = name.Substring(0, dotIndex);
+			}
+
+			return name;
+		}
+
 		/// <inheritdoc/>
 		public override IEnumerable<IssueEventGroup> GetIssues() => _issues;
 	}
diff --git a/Engine/Source/Programs/Shared/EpicGames.LogMatchers.Tests/AutomationToolLogEventParserTests.cs b/Engine/Source/Programs/Shared/EpicGames.LogMatchers.Tests/AutomationToolLogEventParserTests.cs
--- a/Engine/Source/Programs/Shared/EpicGames.LogMatchers.Tests/AutomationToolLogEventParserTests.cs
+++ b/Engine/Source/Programs/Shared/EpicGames.LogMatchers.Tests/AutomationToolLogEventParserTests.cs
@@ -734,6 +734,61 @@ namespace EpicGames.LogMatchers.Tests
 			}
 		}
 
+		[TestMethod]
+		public void ShaderEventMatcherWithFilePaths()
+		{
+			string[] lines =
+			{
+				@"LogCook: Display: Cook Diagnostics: OpenFileHandles=10353, VirtualMemory=20078MiB",
+				@"LogShaderCompilers: Warning: Failed to compile Material /MeshModelingToolsetExp/Materials/PointSetOverlaidComponentMaterial.PointSetOverlaidComponentMaterial for platform SF_PS5, Default Material will be used in game.",
+				@"D:\build\U5M+Inc\Sync\Engine\Shaders\Private\LocalVertexFactory.ush(731:11): Shader FLumenCardVS, Permutation 0, VF FLocalVertexFactory:",
+				@"    /Engine/Private/LocalVertexFactory.ush(731:11): (error, code:5430) - 'FMaterialVertexParameters' does not contain a member named 'TexCoords'",
+				@"D:\build\U5M+Inc\Sync\Engine\Shaders\Private\MaterialTemplate.ush(4282:31): Shader TBasePassVSFNoLightMapPolicy, Permutation 0, VF FLocalVertexFactory:",
+				@"    /Engine/Generated/Material.ush(4282:31): (error, code:5430) - 'FMaterialVertexParameters' does not contain a member named 'TexCoords'",
+				@"Internal Error!",
+				@"    PointSetOverlaidComponentMaterial_6dd3f421e64e95eb/Default/StandardVelocityPipeline/FLocalVertexFactory/TVelocityPS<EVelocityPassMode::Velocity_Standard>/0:/Engine/Private/VelocityShader.usf|MainPixelShader VF 'FLocalVertexFactory' Type 'TVelocityPS<EVelocityPassMode::Velocity_Standard>' '/Engine/Private/VelocityShader.usf' Entry 'MainPixelShader' Permutation 0",
+				@"D:\build\U5M+Inc\Sync\Engine\Shaders\Private\MaterialTemplate.ush(4353:31): Shader TVelocityVS<EVelocityPassMode::Velocity_Standard>, Permutation 0, VF FLocalVertexFactory:",
+				@"    /Engine/Generated/Material.ush(4353:31): (error, code:5430) - 'FMaterialVertexParameters' does not contain a member named 'TexCoords'",
+				@"LogCook: Display: Excluding /Interchange/gltf/MaterialInstances/MI_ClearCoat_Mask_DS"
+			};
+
+			{
+				List<LogEvent> logEvents = Parse(lines);
+				Assert.AreEqual(11, logEvents.Count);
+				CheckEventGroup(logEvents.Slice(1, 9), 1, 9, LogLevel.Warning, KnownLogEvents.Engine_ShaderCompiler);
+
+				// Verify the material property is extracted from "Failed to compile Material" warnings
+				LogEvent shaderEvent = logEvents[1];
+				Assert.IsNotNull(shaderEvent.Properties);
+				object? materialValue;
+				Assert.IsTrue(shaderEvent.TryGetProperty("material", out materialValue));
+				Assert.AreEqual("/MeshModelingToolsetExp/Materials/PointSetOverlaidComponentMaterial.PointSetOverlaidComponentMaterial", materialValue);
+			}
+		}
+
+		[TestMethod]
+		public void ShaderEventMatcherWithoutMaterialProperty()
+		{
+			// Shader errors that don't match "Failed to compile Material" pattern should NOT have material property
+			string[] lines =
+			{
+				@"LogShaderCompilers: Error: Shader compilation failed due to syntax error",
+				@"  error: expected ';' at end of declaration",
+			};
+
+			{
+				List<LogEvent> logEvents = Parse(lines);
+				Assert.AreEqual(2, logEvents.Count);
+				CheckEventGroup(logEvents.Slice(0, 2), 0, 2, LogLevel.Error, KnownLogEvents.Engine_ShaderCompiler);
+
+				// Verify NO material property is set for non-material errors
+				LogEvent shaderEvent = logEvents[0];
+				object? materialValue;
+				bool hasMaterial = shaderEvent.TryGetProperty("material", out materialValue);
+				Assert.IsFalse(hasMaterial);
+			}
+		}
+
 		[TestMethod]
 		public void WorldLeakErrorMatcher()
 		{
diff --git a/Engine/Source/Programs/Shared/EpicGames.LogMatchers/ShaderEventMatcher.cs b/Engine/Source/Programs/Shared/EpicGames.LogMatchers/ShaderEventMatcher.cs
--- a/Engine/Source/Programs/Shared/EpicGames.LogMatchers/ShaderEventMatcher.cs
+++ b/Engine/Source/Programs/Shared/EpicGames.LogMatchers/ShaderEventMatcher.cs
@@ -8,7 +8,9 @@ using Microsoft.Extensions.Logging;
 namespace EpicGames.LogMatchers
 {
 	/// <summary>
-	/// Matches shader compile errors and annotates with the source file path and revision
+	/// Matches shader compile errors and annotates with the source file path and revision.
+	/// Also detects "Failed to compile Material" warnings and extracts the material path
+	/// to enable grouping by material in the issue handler.
 	/// </summary>
 	class ShaderEventMatcher : ILogEventMatcher
 	{
@@ -27,9 +29,24 @@ namespace EpicGames.LogMatchers
 		const string LinePattern =
 			@"(?<line>\d+)(?::(?<column>\d+))?";
 
+		/// <summary>
+		/// Pattern to extract material path from "Failed to compile Material {path}" messages.
+		/// The material path is a forward-slash separated path like "/Game/Materials/MyMaterial.MyMaterial".
+		/// This property enables ShaderIssueHandler to group all compilation failures for the same
+		/// material together, rather than creating separate issues for each shader/platform combination.
+		/// </summary>
+		const string MaterialPattern = @"Failed to compile Material (?<material>/[^\s]+)";
+
 		static readonly Regex s_pattern = new Regex(Prefix);
 		static readonly Regex s_patternWithFile = new Regex($"{Prefix}{FilePattern}(?:\\({LinePattern}\\))?:");
-		static readonly Regex s_content = new Regex(@"^(?:\s+|Validation failed)");
+		static readonly Regex s_materialPattern = new Regex(MaterialPattern);
+
+		// Matches continuation lines that are part of the same shader error:
+		// - Lines starting with whitespace (indented content, error details)
+		// - "Validation failed" messages
+		// - Windows file paths like "D:\build\..." (shader source file references)
+		// - "Internal Error!" messages from shader compiler
+		static readonly Regex s_content = new Regex(@"^(?:\s+|Validation failed|[A-Za-z]:[/\\]|Internal Error!)");
 
 		/// <inheritdoc/>
 		public LogEventMatch? Match(ILogCursor input)
@@ -50,6 +67,18 @@ namespace EpicGames.LogMatchers
 					builder.TryAnnotate(fileMatch.Groups["column"], LogEventMarkup.ColumnNumber);
 				}
 
+				// Check if this is a "Failed to compile Material" warning and extract the material path.
+				// This allows ShaderIssueHandler to group issues by material rather than by source file.
+				string? currentLine = input.CurrentLine;
+				if (currentLine != null)
+				{
+					Match materialMatch = s_materialPattern.Match(currentLine);
+					if (materialMatch.Success)
+					{
+						builder.AddProperty("material", materialMatch.Groups["material"].Value);
+					}
+				}
+
 				while (builder.Next.IsMatch(s_content))
 				{
 					builder.MoveNext();
