[Horde Dashboard] Adding system storage view option to agent page columns

#tests tested with a local built version of the horde dashboard targeting the horde prod environment

diff --git a/Engine/Source/Programs/Horde/HordeDashboard/src/components/agents/AgentView.tsx b/Engine/Source/Programs/Horde/HordeDashboard/src/components/agents/AgentView.tsx
--- a/Engine/Source/Programs/Horde/HordeDashboard/src/components/agents/AgentView.tsx
+++ b/Engine/Source/Programs/Horde/HordeDashboard/src/components/agents/AgentView.tsx
@@ -142,6 +142,32 @@ type ColumnItem = {
    columnDef: IColumn | undefined;
 }
 
+type StorageColumnConfig = {
+   key: string;
+   displayText: string;
+   freeCapability: string;
+   totalCapability: string;
+   searchStateKey: string;
+}
+
+const STORAGE_COLUMNS: Record<string, StorageColumnConfig> = {
+   workingStorage: {
+      key: 'workingStorage',
+      displayText: 'Working Storage',
+      freeCapability: 'WorkingDiskFreeSpace',
+      totalCapability: 'WorkingDiskTotalSize',
+      searchStateKey: 'workingStorage'
+   },
+   systemStorage: {
+      key: 'systemStorage',
+      displayText: 'System Storage',
+      freeCapability: 'SystemDiskFreeSpace',
+      totalCapability: 'SystemDiskTotalSize',
+      searchStateKey: 'systemStorage'
+   }
+};
+
+const BYTES_PER_GIB = 1024 * 1024 * 1024; // 1024^3
 
 function getAgentCapability(agent: AgentData, inProp: string) {
    if (agent.capabilities?.devices) {
@@ -471,7 +497,8 @@ class LocalState {
          // True when colState is for a column type that can be selected for the last column
          const bIsSwapColumn: boolean = (
             colState.key === "software"
-         || colState.key === "storage"
+         || colState.key === "workingStorage"
+         || colState.key === "systemStorage"
          || colState.key.indexOf("systemInfo") !== -1
          || colState.key === "taskTime"
          || colState.key === "comment"
@@ -545,10 +572,17 @@ class LocalState {
          },
          {
             canCheck: true,
-            key: 'storage',
-            checked: this.columnsState.find(c => c.key === 'storage')!.isChecked,
-            text: this.columnsState.find(c => c.key === 'storage')!.displayText,
-            onClick: () => this.setColumnItemChecked('storage')
+            key: 'workingStorage',
+            checked: this.columnsState.find(c => c.key === 'workingStorage')!.isChecked,
+            text: this.columnsState.find(c => c.key === 'workingStorage')!.displayText,
+            onClick: () => this.setColumnItemChecked('workingStorage')
+         },
+         {
+            canCheck: true,
+            key: 'systemStorage',
+            checked: this.columnsState.find(c => c.key === 'systemStorage')!.isChecked,
+            text: this.columnsState.find(c => c.key === 'systemStorage')!.displayText,
+            onClick: () => this.setColumnItemChecked('systemStorage')
          },
          {
             canCheck: true,
@@ -753,8 +787,8 @@ class LocalState {
             columnDef: undefined
          },
          {
-            key: 'storage',
-            displayText: 'Storage',
+            key: 'workingStorage',
+            displayText: 'Working Storage',
             colSize: 220,
             isChecked: true,
             isCheckable: true,
@@ -762,6 +796,16 @@ class LocalState {
             isSortedDescending: false,
             columnDef: undefined
          },
+         {
+            key: 'systemStorage',
+            displayText: 'System Storage',
+            colSize: 220,
+            isChecked: false,
+            isCheckable: true,
+            isSorted: false,
+            isSortedDescending: false,
+            columnDef: undefined
+         },
          {
             key: 'comment',
             displayText: 'Comment',
@@ -816,7 +860,8 @@ class LocalState {
       this.columnSearchState = {
          'pools': {},
          'status': {},
-         'storage': {},
+         'workingStorage': {},
+         'systemStorage': {},
          'comment': {},
          'taskTime': {},
          'taskTimeDescription': {}
@@ -2188,20 +2233,17 @@ const AgentViewInner: React.FC<{ agentId?: string, poolId?: string, searchParams
                   toCheck = "Unknown";
                }
                break;
-            case 'storage':
-               if (localState.columnSearchState['storage'][agent.id]) {
-                  toCheck = localState.columnSearchState['storage'][agent.id];
-               }
-               else {
-                  let diskFree: string | number | null = getAgentCapability(agent, "WorkingDiskFreeSpace");
-                  let diskCapacity: string | number | null = getAgentCapability(agent, "WorkingDiskTotalSize");
-                  toCheck = "Unknown";
-                  if (diskFree !== null && diskCapacity !== null) {
-                     diskFree = Number(diskFree) / 1073741824;
-                     diskCapacity = Number(diskCapacity) / 1073741824;
-                     const percentage = (diskCapacity - diskFree) / diskCapacity;
-                     toCheck = (diskCapacity - (percentage * diskCapacity)).toFixed(0);
-                  }
+            case 'workingStorage':
+            case 'systemStorage':
+               const filterConfig = STORAGE_COLUMNS[localState.searchState?.columnMode ?? 'workingStorage'];
+               if (localState.columnSearchState[filterConfig.searchStateKey][agent.id]) {
+                  toCheck = String(localState.columnSearchState[filterConfig.searchStateKey][agent.id]);
+               } else {
+                  const diskFree = getAgentCapability(agent, filterConfig.freeCapability);
+                  const diskTotal = getAgentCapability(agent, filterConfig.totalCapability);
+                  toCheck = (diskFree !== null && diskTotal !== null)
+                     ? String(getFreeDiskSpaceInGB(Number(diskFree), Number(diskTotal)))
+                     : 'Unknown';
                }
                break;
             case 'comment':
@@ -2243,6 +2285,74 @@ const AgentViewInner: React.FC<{ agentId?: string, poolId?: string, searchParams
       }
    }
 
+   function getStorageMetrics(agent: AgentData, config: StorageColumnConfig): {
+      freeGiB: number;
+      percentUsed: number;
+      displayText: string;
+   } {
+      const diskFree = getAgentCapability(agent, config.freeCapability);
+      const diskTotal = getAgentCapability(agent, config.totalCapability);
+
+      if (diskFree === null || diskTotal === null) {
+         return { freeGiB: 0, percentUsed: 0, displayText: 'Unknown' };
+      }
+
+      const freeGiB = Number(diskFree) / BYTES_PER_GIB;
+      const totalGiB = Number(diskTotal) / BYTES_PER_GIB;
+      const percentUsed = (totalGiB - freeGiB) / totalGiB;
+
+      return {
+         freeGiB,
+         percentUsed,
+         displayText: freeGiB.toLocaleString(undefined, { maximumFractionDigits: 0 }) + ' GiB Free'
+      };
+   }
+
+   function renderStorageCell(agent: AgentData, columnKey: string): JSX.Element {
+      const config = STORAGE_COLUMNS[columnKey];
+      const metrics = getStorageMetrics(agent, config);
+
+      // Update search state for sorting/filtering
+      localState.columnSearchState[config.searchStateKey][agent.id] = metrics.freeGiB;
+
+      // Nudge percentage slightly for progress bar visual
+      const nudgedPercent = metrics.percentUsed * 0.965;
+
+      return (
+         <Stack horizontal horizontalAlign={"center"}>
+            <Stack horizontal horizontalAlign="end" styles={{ root: { minWidth: 80, maxWidth: 80 } }}>
+               <Text styles={{ root: { marginRight: 10, marginTop: 7 } }}>
+                  {metrics.displayText}
+               </Text>
+            </Stack>
+            <Stack horizontal styles={{ root: { minWidth: 110, maxWidth: 110 } }}>
+               <ProgressIndicator
+                  barHeight={15}
+                  styles={{
+                     progressBar: {
+                        height: dashboard.darktheme ? 13 : 16,
+                        marginLeft: dashboard.darktheme ? 1 : 0,
+                        marginTop: dashboard.darktheme ? 2 : 0
+                     },
+                     progressTrack: { width: '98%', border: '1px solid !important' },
+                     root: { width: 110 }
+                  }}
+                  percentComplete={nudgedPercent}
+               />
+            </Stack>
+         </Stack>
+      );
+   }
+
+   function getFreeDiskSpaceInGB(diskFree: number, diskCapacity: number): string | undefined {
+      if (diskFree === null || diskCapacity === null) {
+         return undefined;
+      }
+
+      const freeGB = diskFree / BYTES_PER_GIB;
+      return freeGB.toFixed(0);
+   }
+
    function sortAgents(a: AgentData, b: AgentData) {
       // sort differently depending on which column this is
       const sortedCol = localState.columnsState.find(c => c.isSorted === true)!;
@@ -2273,11 +2383,15 @@ const AgentViewInner: React.FC<{ agentId?: string, poolId?: string, searchParams
             return (left.version ?? "").toLocaleLowerCase().localeCompare(right.version ?? "");
          case 'taskTime':
             return getTaskTime(left) - getTaskTime(right);
-         case 'storage':
-            if (localState.columnSearchState['storage'][left.id] && localState.columnSearchState['storage'][right.id]) {
-               return localState.columnSearchState['storage'][left.id] - localState.columnSearchState['storage'][right.id];
+         case 'workingStorage':
+         case 'systemStorage':
+            const sortConfig = STORAGE_COLUMNS[sortedCol.key];
+            const leftVal = localState.columnSearchState[sortConfig.searchStateKey][left.id];
+            const rightVal = localState.columnSearchState[sortConfig.searchStateKey][right.id];
+            if (leftVal !== undefined && rightVal !== undefined) {
+               return leftVal - rightVal;
             }
-            return Number(getAgentCapability(left, "DiskFreeSpace")) - Number(getAgentCapability(right, "DiskFreeSpace"));
+            return Number(getAgentCapability(left, sortConfig.freeCapability) ?? 0) - Number(getAgentCapability(right, sortConfig.freeCapability) ?? 0);
          case 'comment':
             if (localState.columnSearchState['comment'][left.id] && localState.columnSearchState['comment'][right.id]) {
                return localState.columnSearchState['comment'][left.id].localeCompare(localState.columnSearchState['comment'][right.id]);
@@ -2875,40 +2989,9 @@ const AgentViewInner: React.FC<{ agentId?: string, poolId?: string, searchParams
                   </Text>
                </Stack.Item>
             </Stack>);
-         case 'storage':
-            let diskFree: string | number | null = getAgentCapability(agent, "WorkingDiskFreeSpace");
-            let diskCapacity: string | number | null = getAgentCapability(agent, "WorkingDiskTotalSize");
-
-            let realPercentage = 0;
-            let nudgedPercentage = 0;
-            let descVal = "Unknown";
-            let descNumber = 0;
-            if (diskFree !== null && diskCapacity !== null) {
-               diskFree = Number(diskFree) / 1073741824;
-               diskCapacity = Number(diskCapacity) / 1073741824;
-               realPercentage = (diskCapacity - diskFree) / diskCapacity;
-               // nudge percentage for the progress bar length
-               nudgedPercentage = realPercentage * .965;
-               descNumber = (diskCapacity - (realPercentage * diskCapacity));
-               descVal = descNumber.toLocaleString(undefined, { maximumFractionDigits: 0 }) + " GiB Free";
-            }
-            localState.columnSearchState['storage'][agent.id] = descNumber;
-            return <Stack horizontal horizontalAlign={"center"}>
-               <Stack horizontal horizontalAlign={"end"} styles={{ root: { minWidth: 80, maxWidth: 80 } }}>
-                  <Text styles={{ root: { marginRight: 10, marginTop: 7 } }}>{descVal}</Text>
-               </Stack>
-               <Stack horizontal styles={{ root: { minWidth: 110, maxWidth: 110 } }}>
-                  <ProgressIndicator
-                     barHeight={15}
-                     styles={{
-                        progressBar: { height: dashboard.darktheme ? 13 : 16, marginLeft: dashboard.darktheme ? 1 : 0, marginTop: dashboard.darktheme ? 2 : 0 },
-                        progressTrack: { width: '98%', border: '1px solid !important' },
-                        root: { width: 110 }
-                     }}
-                     percentComplete={nudgedPercentage}
-                  />
-               </Stack>
-            </Stack>;
+         case 'workingStorage':
+         case 'systemStorage':
+            return renderStorageCell(agent, column!.key);
          case 'comment':
             localState.columnSearchState['comment'][agent.id] = agent.comment ?? "";
             return (<Stack styles={{ root: { height: '100%', } }} horizontal horizontalAlign={"center"}>
