Add deep link controller

#rnx

#Synced-CL 50707890

diff --git a/Engine/Source/Programs/Horde/Plugins/Build/HordeServer.Build.Tests/Notifications/NotificationServiceTests.cs b/Engine/Source/Programs/Horde/Plugins/Build/HordeServer.Build.Tests/Notifications/NotificationServiceTests.cs
--- a/Engine/Source/Programs/Horde/Plugins/Build/HordeServer.Build.Tests/Notifications/NotificationServiceTests.cs
+++ b/Engine/Source/Programs/Horde/Plugins/Build/HordeServer.Build.Tests/Notifications/NotificationServiceTests.cs
@@ -96,6 +96,8 @@ namespace HordeServer.Tests.Notifications
 		public Task SendIssueReportAsync(IssueReportGroup report, CancellationToken cancellationToken) => throw new NotImplementedException();
 		public Task SendDeviceIssueReportAsync(DeviceIssueReport report, CancellationToken cancellationToken) => throw new NotImplementedException();
 		public Task NotifyTestHealthReportAsync(ITestHealthReport report, string recipient, string[]? carbonCopies, CancellationToken cancellationToken) => throw new NotImplementedException();
+		public Task<string?> GetDirectMessageLinkAsync(IReadOnlyList<UserId> userIds, CancellationToken cancellationToken = default) => throw new NotImplementedException();
+		public Task<string?> GetChannelLinkAsync(string channel, CancellationToken cancellationToken = default) => throw new NotImplementedException();
 	}
 
 	[TestClass]
diff --git a/Engine/Source/Programs/Horde/Plugins/Build/HordeServer.Build/Notifications/DeepLinkController.cs b/Engine/Source/Programs/Horde/Plugins/Build/HordeServer.Build/Notifications/DeepLinkController.cs
new file mode 100644
--- /dev/null
+++ b/Engine/Source/Programs/Horde/Plugins/Build/HordeServer.Build/Notifications/DeepLinkController.cs
@@ -0,0 +1,109 @@
+// Copyright Epic Games, Inc. All Rights Reserved.
+
+using EpicGames.Horde.Users;
+using HordeServer.Utilities;
+using Microsoft.AspNetCore.Authorization;
+using Microsoft.AspNetCore.Mvc;
+using Microsoft.Extensions.Logging;
+
+namespace HordeServer.Notifications
+{
+	/// <summary>
+	/// Controller for the /api/v1/deeplink endpoint
+	/// </summary>
+	[ApiController]
+	[Authorize]
+	[Route("[controller]")]
+	public class DeepLinkController : ControllerBase
+	{
+		private readonly INotificationService _notificationService;
+		private readonly ILogger _logger;
+
+		/// <summary>
+		/// Constructor
+		/// </summary>
+		public DeepLinkController(INotificationService notificationService, ILogger<DeepLinkController> logger)
+		{
+			_notificationService = notificationService;
+			_logger = logger;
+		}
+
+		/// <summary>
+		/// Get direct message deep link
+		/// </summary>
+		/// <param name="request"></param>
+		/// <param name="cancellationToken">Cancellation token for the operation</param>
+		/// <returns></returns>
+		[HttpPost]
+		[Route("/api/v1/deeplink/directmessage")]
+		[ProducesResponseType(typeof(GetDirectMessageLinkResponse), 200)]
+		public async Task<ActionResult<GetDirectMessageLinkResponse>> GetDirectMessageDeepLinkAsync([FromBody] GetDirectMessageLinkRequest request, CancellationToken cancellationToken = default)
+		{
+			HashSet<UserId>? userIds = request?.UserIds?.Select(u => TryParseUserId(u)).OfType<UserId>().ToHashSet();
+			if (userIds is null || userIds.Count == 0)
+			{
+				_logger.LogWarning("Invalid user ids: {UserIds}", request?.UserIds == null? "null" : String.Join(", ", request.UserIds));
+				return BadRequest("Invalid user ids");
+			}
+
+			UserId? userId = User.GetUserId();
+			if (userId is null)
+			{
+				return BadRequest("No user logged in.");
+			}
+
+			userIds.Add(userId!.Value);
+
+			string? url = await _notificationService.GetDirectMessageLinkAsync(userIds.ToList(), cancellationToken);
+			if (url is null)
+			{
+				return Problem("The DM could not be deep linked");
+			}
+
+			return new GetDirectMessageLinkResponse() { Url = url };
+		}
+
+		/// <summary>
+		/// Get channel deep link
+		/// </summary>
+		/// <param name="channelId"></param>
+		/// <param name="cancellationToken">Cancellation token for the operation</param>
+		/// <returns></returns>
+		[HttpGet]
+		[Route("/api/v1/deeplink/channel")]
+		[ProducesResponseType(typeof(GetChannelLinkResponse), 200)]
+		public async Task<ActionResult<GetChannelLinkResponse>> GetChannelDeepLinkAsync([FromQuery] string channelId, CancellationToken cancellationToken = default)
+		{
+			if (String.IsNullOrEmpty(channelId))
+			{
+				return BadRequest("Channel id is invalid");
+			}
+
+			string? url = await _notificationService.GetChannelLinkAsync(channelId, cancellationToken);
+			if (url is null)
+			{
+				return NotFound("Could not find the channel");
+			}
+
+			return new GetChannelLinkResponse() { Url = url };
+		}
+
+		/// <summary>
+		/// Parse a user id from a string. Return null if failed to parse.
+		/// </summary>
+		/// <param name="userId"></param>
+		/// <returns></returns>
+		static UserId? TryParseUserId(string userId)
+		{
+			UserId newObjectId;
+			if (!String.IsNullOrEmpty(userId) && UserId.TryParse(userId, out newObjectId))
+			{
+				return newObjectId;
+			}
+			else
+			{
+				return null;
+			}
+		}
+	}
+}
diff --git a/Engine/Source/Programs/Horde/Plugins/Build/HordeServer.Build/Notifications/DeepLinkMessages.cs b/Engine/Source/Programs/Horde/Plugins/Build/HordeServer.Build/Notifications/DeepLinkMessages.cs
new file mode 100644
--- /dev/null
+++ b/Engine/Source/Programs/Horde/Plugins/Build/HordeServer.Build/Notifications/DeepLinkMessages.cs
@@ -0,0 +1,37 @@
+// Copyright Epic Games, Inc. All Rights Reserved.
+
+namespace HordeServer.Notifications
+{
+	/// <summary>
+	/// A direct message link request
+	/// </summary>
+	public class GetDirectMessageLinkRequest
+	{
+		/// <summary>
+		/// the user ids
+		/// </summary>
+		public List<string> UserIds { get; set; } = new List<string>();
+	}
+
+	/// <summary>
+	/// The Direct Message link response
+	/// </summary>
+	public class GetDirectMessageLinkResponse
+	{
+		/// <summary>
+		/// The result deep link
+		/// </summary>
+		public string? Url { get; set; }
+	}
+
+	/// <summary>
+	/// The Channel link response
+	/// </summary>
+	public class GetChannelLinkResponse
+	{
+		/// <summary>
+		/// The result deep link
+		/// </summary>
+		public string? Url { get; set; }
+	}
+}
\ No newline at end of file
diff --git a/Engine/Source/Programs/Horde/Plugins/Build/HordeServer.Build/Notifications/INotificationService.cs b/Engine/Source/Programs/Horde/Plugins/Build/HordeServer.Build/Notifications/INotificationService.cs
--- a/Engine/Source/Programs/Horde/Plugins/Build/HordeServer.Build/Notifications/INotificationService.cs
+++ b/Engine/Source/Programs/Horde/Plugins/Build/HordeServer.Build/Notifications/INotificationService.cs
@@ -5,6 +5,7 @@ using System.Security.Claims;
 using EpicGames.Horde.Jobs;
 using EpicGames.Horde.Jobs.Graphs;
 using EpicGames.Horde.Notifications;
+using EpicGames.Horde.Users;
 using HordeServer.Agents;
 using HordeServer.Configuration;
 using HordeServer.Devices;
@@ -198,11 +199,27 @@ namespace HordeServer.Notifications
 		/// <summary>
 		/// Post a notification for test health report
 		/// </summary>
-		/// <param name="report"></param>
-		/// <param name="recipient"></param>
-		/// <param name="carbonCopies"></param>
-		/// <param name="cancellationToken"></param>
+		/// <param name="report">The report data to send</param>
+		/// <param name="recipient">The recipient of the report</param>
+		/// <param name="carbonCopies">The users to send a carbon copy</param>
+		/// <param name="cancellationToken">Cancellation token for the operation</param>
 		/// <returns></returns>
 		Task NotifyTestHealthReportAsync(ITestHealthReport report, string recipient, string[]? carbonCopies, CancellationToken cancellationToken);
+
+		/// <summary>
+		/// Get Direct Message link
+		/// </summary>
+		/// <param name="userIds">The user ids to start conversation with</param>
+		/// <param name="cancellationToken">Cancellation token for the operation</param>
+		/// <returns></returns>
+		Task<string?> GetDirectMessageLinkAsync(IReadOnlyList<UserId> userIds, CancellationToken cancellationToken = default);
+
+		/// <summary>
+		/// Get Channel link
+		/// </summary>
+		/// <param name="channel">The channel to get the link for</param>
+		/// <param name="cancellationToken">Cancellation token for the operation</param>
+		/// <returns></returns>
+		Task<string?> GetChannelLinkAsync(string channel, CancellationToken cancellationToken = default);
 	}
 }
diff --git a/Engine/Source/Programs/Horde/Plugins/Build/HordeServer.Build/Notifications/INotificationSink.cs b/Engine/Source/Programs/Horde/Plugins/Build/HordeServer.Build/Notifications/INotificationSink.cs
--- a/Engine/Source/Programs/Horde/Plugins/Build/HordeServer.Build/Notifications/INotificationSink.cs
+++ b/Engine/Source/Programs/Horde/Plugins/Build/HordeServer.Build/Notifications/INotificationSink.cs
@@ -2,6 +2,7 @@
 
 using EpicGames.Horde.Jobs;
 using EpicGames.Horde.Jobs.Graphs;
+using EpicGames.Horde.Users;
 using HordeServer.Agents;
 using HordeServer.Configuration;
 using HordeServer.Devices;
@@ -161,9 +162,25 @@ namespace HordeServer.Notifications
 		/// <param name="report"></param>
 		/// <param name="recipient"></param>
 		/// <param name="carbonCopies"></param>
-		/// <param name="cancellationToken"></param>
+		/// <param name="cancellationToken">Cancellation token for the operation</param>
 		/// <returns></returns>
 		Task NotifyTestHealthReportAsync(ITestHealthReport report, string recipient, string[]? carbonCopies, CancellationToken cancellationToken);
+
+		/// <summary>
+		/// Get Direct Message Link
+		/// </summary>
+		/// <param name="userIds"></param>
+		/// <param name="cancellationToken">Cancellation token for the operation</param>
+		/// <returns></returns>
+		Task<string?> GetDirectMessageLinkAsync(IReadOnlyList<UserId> userIds, CancellationToken cancellationToken = default);
+
+		/// <summary>
+		/// Get Channel link
+		/// </summary>
+		/// <param name="channel"></param>
+		/// <param name="cancellationToken">Cancellation token for the operation</param>
+		/// <returns></returns>
+		Task<string?> GetChannelLinkAsync(string channel, CancellationToken cancellationToken = default);
 	}
 }
 
diff --git a/Engine/Source/Programs/Horde/Plugins/Build/HordeServer.Build/Notifications/NotificationService.cs b/Engine/Source/Programs/Horde/Plugins/Build/HordeServer.Build/Notifications/NotificationService.cs
--- a/Engine/Source/Programs/Horde/Plugins/Build/HordeServer.Build/Notifications/NotificationService.cs
+++ b/Engine/Source/Programs/Horde/Plugins/Build/HordeServer.Build/Notifications/NotificationService.cs
@@ -904,5 +904,64 @@ namespace HordeServer.Notifications
 				}
 			}
 		}
+
+		/// <inheritdoc/>
+		/// <remarks>It will return the first non null output from the sinks ignoring the rest</remarks>
+		public async Task<string?> GetDirectMessageLinkAsync(IReadOnlyList<UserId> userIds, CancellationToken cancellationToken)
+		{
+			int sinksLeft = _sinks.Count;
+			foreach (INotificationSink sink in _sinks)
+			{
+				sinksLeft--;
+				try
+				{
+					string? url = await sink.GetDirectMessageLinkAsync(userIds, cancellationToken);
+					if (url is not null)
+					{
+						if (sinksLeft > 0)
+						{
+							_logger.LogInformation("A DM link was generated from {Sink} ignoring the remaining {Count} sink(s)", sink.GetType().Name, sinksLeft);
+						}
+						return url;
+					}
+				}
+				catch (Exception e)
+				{
+					string users = String.Join(", ", userIds.Select(x => x.ToString()));
+					_logger.LogError(e, "Failed getting direct message link for users {Users}", users);
+				}
+			}
+
+			return null;
+		}
+
+		/// <inheritdoc/>
+		/// <remarks>It will return the first non null output from the sinks ignoring the rest</remarks>
+		public async Task<string?> GetChannelLinkAsync(string channel, CancellationToken cancellationToken)
+		{
+			int sinksLeft = _sinks.Count;
+			foreach (INotificationSink sink in _sinks)
+			{
+				sinksLeft--;
+				try
+				{
+					string? url = await sink.GetChannelLinkAsync(channel, cancellationToken);
+					if (url is not null)
+					{
+						if (sinksLeft > 0)
+						{
+							_logger.LogInformation("A Channel link was generated from {Sink} ignoring the remaining {Count} sink(s)", sink.GetType().Name, sinksLeft);
+						}
+						return url;
+					}
+				}
+				catch (Exception e)
+				{
+					_logger.LogError(e, "Failed getting channel link for channel {Channel}", channel);
+				}
+			}
+
+			return null;
+		}
 	}
 }
diff --git a/Engine/Source/Programs/Horde/Plugins/Build/HordeServer.Build/Notifications/Sinks/SlackNotificationSink.cs b/Engine/Source/Programs/Horde/Plugins/Build/HordeServer.Build/Notifications/Sinks/SlackNotificationSink.cs
--- a/Engine/Source/Programs/Horde/Plugins/Build/HordeServer.Build/Notifications/Sinks/SlackNotificationSink.cs
+++ b/Engine/Source/Programs/Horde/Plugins/Build/HordeServer.Build/Notifications/Sinks/SlackNotificationSink.cs
@@ -5,6 +5,7 @@ using System.Diagnostics.CodeAnalysis;
 using System.Globalization;
 using System.Net.WebSockets;
 using System.Text;
+using System.Text.Encodings.Web;
 using System.Text.Json;
 using System.Text.Json.Serialization;
 using System.Text.RegularExpressions;
@@ -246,6 +247,11 @@ namespace HordeServer.Notifications.Sinks
 		/// </summary>
 		private readonly MemoryCache _userCache = new MemoryCache(new MemoryCacheOptions());
 
+		/// <summary>
+		/// Map of Slack user IDs to DM channel.
+		/// </summary>
+		private readonly MemoryCache _dmChannelCache = new MemoryCache(new MemoryCacheOptions() { SizeLimit = 1024, ExpirationScanFrequency = TimeSpan.FromMinutes(10) });
+
 		/// <summary>
 		/// Constructor
 		/// </summary>
@@ -292,6 +298,7 @@ namespace HordeServer.Notifications.Sinks
 		{
 			await _backgroundTask.DisposeAsync();
 			_userCache.Dispose();
+			_dmChannelCache.Dispose();
 			_slackClient.Dispose();
 			_adminSlackClient?.Dispose();
 			await _issueQueueTicker.DisposeAsync();
@@ -3954,5 +3961,70 @@ namespace HordeServer.Notifications.Sinks
 				}
 			}
 		}
+
+		const string DeepLinkUrl = "https://slack.com/app_redirect";
+		const int MaxDirectMessageUsers = 8;
+
+		/// <inheritdoc/>
+		public async Task<string?> GetDirectMessageLinkAsync(IReadOnlyList<UserId> userIds, CancellationToken cancellationToken)
+		{
+			if (userIds.Count == 0 || userIds.Count > MaxDirectMessageUsers)
+			{
+				_logger.LogWarning("Number of users is out of bound: {Count}. Slack supports 1 to 8 users for DM conversation.", userIds.Count);
+				return null;
+			}
+
+			string?[] taskResults = await Task.WhenAll(userIds.Distinct().Select(u => GetSlackUserIdAsync(u, cancellationToken)));
+			List<string> slackUserIds = taskResults.OfType<string>().ToList();
+			if (slackUserIds.Count == 0)
+			{
+				_logger.LogWarning("Could not retrieve enough valid Slack user ids for DM conversation.");
+				return null;
+			}
+			slackUserIds.Sort();
+			string slackUserIdsString = String.Join(",", slackUserIds);
+
+			string? channelId;
+			if (!_dmChannelCache.TryGetValue(slackUserIdsString, out channelId))
+			{
+				channelId = await _slackClient.OpenDirectMessageAsync(slackUserIds, cancellationToken);
+				if (channelId is null)
+				{
+					_logger.LogWarning("Could not open DM conversation for users {Users}.", slackUserIdsString);
+					return null;
+				}
+				using (ICacheEntry entry = _dmChannelCache.CreateEntry(slackUserIdsString))
+				{
+					entry.SlidingExpiration = TimeSpan.FromMinutes(10.0);
+					entry.Value = channelId;
+					entry.Size = 1;
+				}
+			}
+
+			return $"{DeepLinkUrl}?channel={channelId}";
+		}
+
+		async Task<string?> GetSlackUserIdAsync(UserId userId, CancellationToken cancellationToken)
+		{
+			IUser? user = await _userCollection.GetUserAsync(userId, cancellationToken);
+			if (user == null)
+			{
+				return null;
+			}
+
+			string? slackUserId = await GetSlackUserIdAsync(user, cancellationToken);
+			if (slackUserId == null)
+			{
+				return null;
+			}
+
+			return slackUserId;
+		}
+
+		/// <inheritdoc/>
+		public Task<string?> GetChannelLinkAsync(string channelId, CancellationToken cancellationToken)
+		{
+			return Task.FromResult<string?>($"{DeepLinkUrl}?channel={UrlEncoder.Default.Encode(channelId)}");
+		}
 	}
 }
diff --git a/Engine/Source/Programs/Horde/Plugins/Experimental/HordeServer.Experimental.Tests/Notifications/ExperimentalNotificationTests.cs b/Engine/Source/Programs/Horde/Plugins/Experimental/HordeServer.Experimental.Tests/Notifications/ExperimentalNotificationTests.cs
--- a/Engine/Source/Programs/Horde/Plugins/Experimental/HordeServer.Experimental.Tests/Notifications/ExperimentalNotificationTests.cs
+++ b/Engine/Source/Programs/Horde/Plugins/Experimental/HordeServer.Experimental.Tests/Notifications/ExperimentalNotificationTests.cs
@@ -129,6 +129,8 @@ namespace HordeServer.Experimental.Tests.Notifications
 		public Task SendIssueReportAsync(IssueReportGroup report, CancellationToken cancellationToken) => throw new NotImplementedException();
 		public Task SendDeviceIssueReportAsync(DeviceIssueReport report, CancellationToken cancellationToken) => throw new NotImplementedException();
 		public Task NotifyTestHealthReportAsync(ITestHealthReport report, string recipient, string[]? carbonCopies, CancellationToken cancellationToken) => throw new NotImplementedException();
+		public Task<string?> GetDirectMessageLinkAsync(IReadOnlyList<UserId> userIds, CancellationToken cancellationToken = default) => throw new NotImplementedException();
+		public Task<string?> GetChannelLinkAsync(string channel, CancellationToken cancellationToken = default) => throw new NotImplementedException();
 		public void Dispose()
 		{
 			Dispose(true);
diff --git a/Engine/Source/Programs/Horde/Plugins/Experimental/HordeServer.Experimental/Notifications/Sinks/ExperimentalSlackNotificationSink.cs b/Engine/Source/Programs/Horde/Plugins/Experimental/HordeServer.Experimental/Notifications/Sinks/ExperimentalSlackNotificationSink.cs
--- a/Engine/Source/Programs/Horde/Plugins/Experimental/HordeServer.Experimental/Notifications/Sinks/ExperimentalSlackNotificationSink.cs
+++ b/Engine/Source/Programs/Horde/Plugins/Experimental/HordeServer.Experimental/Notifications/Sinks/ExperimentalSlackNotificationSink.cs
@@ -2,6 +2,7 @@
 
 using EpicGames.Horde.Jobs;
 using EpicGames.Horde.Jobs.Graphs;
+using EpicGames.Horde.Users;
 using HordeServer.Agents;
 using HordeServer.Configuration;
 using HordeServer.Devices;
@@ -121,6 +122,14 @@ namespace HordeServer.Notifications.Sinks
 		public Task NotifyTestHealthReportAsync(ITestHealthReport report, string recipient, string[]? carbonCopies, CancellationToken cancellationToken)
 			=> Task.CompletedTask;
 
+		/// <inheritdoc/>
+		public Task<string?> GetDirectMessageLinkAsync(IReadOnlyList<UserId> userIds, CancellationToken cancellationToken = default)
+			=> Task.FromResult<string?>(null);
+
+		/// <inheritdoc/>
+		public Task<string?> GetChannelLinkAsync(string channel, CancellationToken cancellationToken = default)
+			=> Task.FromResult<string?>(null);
+
 		/// <inheritdoc/>
 		public void Dispose()
 			=> _slackNotificationProcessor.Dispose();
diff --git a/Engine/Source/Programs/Shared/EpicGames.Slack/SlackClient.cs b/Engine/Source/Programs/Shared/EpicGames.Slack/SlackClient.cs
--- a/Engine/Source/Programs/Shared/EpicGames.Slack/SlackClient.cs
+++ b/Engine/Source/Programs/Shared/EpicGames.Slack/SlackClient.cs
@@ -471,6 +471,7 @@ namespace EpicGames.Slack
 		#region Conversations
 
 		const string ConversationsInviteUrl = "https://slack.com/api/conversations.invite";
+		const string ConversationsOpenUrl = "https://slack.com/api/conversations.open";
 
 		class InviteMessage
 		{
@@ -481,6 +482,24 @@ namespace EpicGames.Slack
 			public string? Users { get; set; } // Comma separated list of ids
 		}
 
+		class OpenDirectMessage
+		{
+			[JsonPropertyName("users")]
+			public string? Users { get; set; } // Comma separated list of ids
+		}
+
+		class GetChannelResponse : SlackResponse
+		{
+			[JsonPropertyName("channel")]
+			public SlackChannelObject? Channel { get; set; }
+		}
+
+		class SlackChannelObject
+		{
+			[JsonPropertyName("id")]
+			public string Id { get; set; } = String.Empty;
+		}
+
 		/// <summary>
 		/// Invite a user to a channel
 		/// </summary>
@@ -536,6 +555,26 @@ namespace EpicGames.Slack
 			await SendRequestAsync<SlackResponse>(AdminConversationsInviteUrl, message, cancellationToken);
 		}
 
+		/// <summary>
+		/// Open a direct message with multiple users
+		/// </summary>
+		/// <param name="userIds">The list of user ids</param>
+		/// <param name="cancellationToken"></param>
+		/// <returns></returns>
+		public async Task <string?> OpenDirectMessageAsync(IEnumerable<string> userIds, CancellationToken cancellationToken = default)
+		{
+			OpenDirectMessage message = new OpenDirectMessage();
+			message.Users = String.Join(",", userIds);
+
+			GetChannelResponse response = await SendRequestAsync<GetChannelResponse>(ConversationsOpenUrl, message, cancellationToken);
+			if (response.Channel is null)
+			{
+				return null;
+			}
+
+			return response.Channel.Id;
+		}
+
 		#endregion
 
 		#region Users
