Add Slack deep linking for channel and direct message backend calls
+ [Test Health] Add external links to workflow report channel and opening direct message with owner and customers

#rnx
#Synced-CL 50707890

diff --git a/Engine/Source/Programs/Horde/HordeDashboard/plugins/test-automation/testAutomationHealth.tsx b/Engine/Source/Programs/Horde/HordeDashboard/plugins/test-automation/testAutomationHealth.tsx
--- a/Engine/Source/Programs/Horde/HordeDashboard/plugins/test-automation/testAutomationHealth.tsx
+++ b/Engine/Source/Programs/Horde/HordeDashboard/plugins/test-automation/testAutomationHealth.tsx
@@ -12,7 +12,9 @@ import { projectStore } from "horde/backend/ProjectStore";
 import dashboard, { StatusColor } from "horde/backend/Dashboard";
 import { getShortNiceTime, msecToElapsed } from "horde/base/utilities/timeUtils";
 import { booleanColors, cadenceColors, cadenceOrders, healthColors, healthStars, TestCadence, TestHealth, TestHealthExplain, TestHealthItem, TestHealthTiming, UsersInfoCache } from "./testHealthData";
-import { GetUserResponse } from "horde/backend/Api";
+import { GetUserResponse, DeepLinkResponse } from "horde/backend/Api";
+import { errorDialogStore } from "horde/components/error/ErrorStore";
+import backend from "horde/backend";
 
 /// Colors
 const profileColors = new DomainNameColors();
@@ -168,7 +170,7 @@ const getHealthColumns = (items: TestHealthItem[], columns: TestHealthColumn[],
             isSorted: false,
             isSortedDescending: true,
             onRender: (item: TestHealthItem) => {
-                const profile = item.profile ?? "N/A";
+                const profile = !!item.profile ? item.profile : "N/A";
                 return <Stack style={{...healthStyles.badge, backgroundColor: profileColors.get(item.profile)}} horizontalAlign="center" horizontal>
                             <Text title={profile} style={{textOverflow: 'ellipsis', overflow: 'hidden', whiteSpace: 'nowrap'}}>{profile}</Text>
                         </Stack>
@@ -186,7 +188,7 @@ const getHealthColumns = (items: TestHealthItem[], columns: TestHealthColumn[],
             isSorted: false,
             isSortedDescending: true,
             onRender: (item: TestHealthItem) => {
-                const team = item.team  ?? "N/A";
+                const team = !!item.team ? item.team : "N/A";
                 return <Stack style={{...healthStyles.badge, backgroundColor: teamColors.get(item.team)}} horizontalAlign="center" horizontal>
                             <Text title={team} style={{textOverflow: 'ellipsis', overflow: 'hidden', whiteSpace: 'nowrap'}}>{team}</Text>
                         </Stack>
@@ -334,9 +336,10 @@ const getTestHealthItems = (handler: TestDataHandler, stream?: string, test?: Te
 
     const filteredMetadata = handler.filteredMetadata;
     
-    // collect data and reorganize it for compution
+    // collect data and reorganize it for computation
     const testHealths: TestHealthItem[] = [];
     const testMeta: Map<string, [TestNameRef, MetadataRef[]][]> = new Map();
+    let testMetaKeys: Map<string, string[]> | undefined;
 
     filteredTests.forEach((test) => {
         const status = testStatus!.tests.get(test);
@@ -429,13 +432,13 @@ const getTestHealthItems = (handler: TestDataHandler, stream?: string, test?: Te
         let healthValue = TestHealth.NA;
         let healthExplain: TestHealthExplain | undefined;
         if (totalPhaseCount > 0) {
-            const catastrophicFailureRate = totalFailureCount == 0 ? 0: percRound(totalUndefinedCount/totalFailureCount);
+            const catastrophicFailureRate = totalFailureCount === 0 ? 0: percRound(totalUndefinedCount/totalFailureCount);
             const failureRate = percRound(totalFailureCount/totalPhaseCount);
             const successRate = percRound(totalSuccessCount/totalPhaseCount);
-            const redundantErrorRate = totalFailureCount == 0 ? 0: percRound(1 - errorFingerprints.size/totalFailureCount);
+            const redundantErrorRate = totalFailureCount === 0 ? 0: percRound(1 - errorFingerprints.size/totalFailureCount);
             let nStars = healthStars.get(TestHealth.Fair)!;
             if (failureRate > 15) {
-                if (catastrophicFailureRate > 40 || successRate == 0) {
+                if (catastrophicFailureRate > 40 || successRate === 0) {
                     nStars -= 2;
                 } else {
                     nStars -= 1;
@@ -478,14 +481,14 @@ const getTestHealthItems = (handler: TestDataHandler, stream?: string, test?: Te
         testHealthItem.health = healthValue;
         testHealthItem.healthExplain = healthExplain;
         testHealthItem.cadence = cadenceValue;
-        testHealthItem.setLazyDisplayMeta(() => status.sessions.keys().find((meta) => !!meta)?.getCommonValues(testMetaKeys.get(test.name))?.join(" / "));
+        testHealthItem.setLazyDisplayMeta(() => status.sessions.keys().find((meta) => !!meta)?.getCommonValues(testMetaKeys?.get(test.name))?.join(" / "));
         testHealthItem.timings = timings;
         testHealthItem.linkRecipeIds(recipeIds);
 
         testHealths.push(testHealthItem);
     });
 
-    const testMetaKeys = getTestUniqueMetaIdentifiers(testMeta);
+    testMetaKeys = getTestUniqueMetaIdentifiers(testMeta);
 
     return testHealths;
 }
@@ -495,7 +498,7 @@ const TestHealthTimes: React.FC<{ test: TestHealthItem }> = ({test}) => {
 
     useEffect(() => {
         test.fetchRecipes().then(() => setInit(true));
-    }, []);
+    }, [test]);
 
     if (!test.timings) return null;
 
@@ -528,15 +531,33 @@ const TestHealthTimes: React.FC<{ test: TestHealthItem }> = ({test}) => {
 
 const usersInfoCache = new UsersInfoCache();
 
+type DeepLinkState = {
+    directMessage: {
+        isFetching?: boolean;
+        url?: string;
+    };
+    notification: {
+        isFetching?: boolean;
+        url?: string;
+    };
+}
+
 export const TestHealthView: React.FC<{ test: TestHealthItem, handler: TestDataHandler, onUpdateItem: (item: TestHealthItem) => void }> = observer(({ test, handler, onUpdateItem }) => {
     const [streams, setStreams] = useState<{current: string, toCompare?: string}>(
         {
             current: handler.stream!,
             toCompare: userStorage.getItem(StoreKey.HistoryToCompare, () => undefined)
         });
-    const [testHealthToCompare, setTestHealthToCompare] = useState<TestHealthItem | undefined>(getTestHealthItems(handler, userStorage.getItem(StoreKey.HistoryToCompare), test.nameRef).find(() => true));
+    const [testHealthToCompare, setTestHealthToCompare] = useState<TestHealthItem | undefined>(getTestHealthItems(handler, userStorage.getItem(StoreKey.HistoryToCompare), test.nameRef).at(0));
     const [userPicker, setUserPicker] = useState<{target?: UserType, defaultUser?: GetUserResponse}>({});
+    const [deepLinkState, setDeepLinkState] = useState<DeepLinkState>({directMessage: {}, notification: {}});
 
+    // reset deep link notification state when notification changes
+    useEffect(() => {
+        setDeepLinkState(prev => ({...prev, notification: {}}));
+    }, [test.notification]);
+
+    // fetch users info when test changes
     useEffect(() => {
         let userIds: string[] = [];
         if (test.user) {
@@ -551,7 +572,9 @@ export const TestHealthView: React.FC<{ test: TestHealthItem, handler: TestDataH
         if (userIds.length > 0) {
             usersInfoCache.fetchUsersInfo(userIds).then((fetchedUsers) => fetchedUsers.length > 0 && test.setUpdated());
         }
-    }, [test.owner, test.customers]);
+        // reset deep link direct message state when users info changes
+        setDeepLinkState(prev => ({...prev, directMessage: {}}));
+    }, [test.user,test.owner, test.customers]);
 
     handler.subscribeToAuditsLoading();
     handler.subscribeToRecipesLoading();
@@ -566,19 +589,19 @@ export const TestHealthView: React.FC<{ test: TestHealthItem, handler: TestDataH
         }
     }
 
-    const commonMetaKeys = metadataSet.size == 1? handler.filterState.metadata?.keys().toArray() : MetadataRef.identifyCommonKeys(Array.from(metadataSet));
+    const commonMetaKeys = metadataSet.size === 1? handler.filterState.metadata?.keys().toArray() : MetadataRef.identifyCommonKeys(Array.from(metadataSet));
     const availableStreams = handler.availableStreams.filter((s) => s !== streams.current);
 
     const onClickStream = useCallback((targetStream: string) => {
-        setStreams({...streams, toCompare: undefined});
+        setStreams(prev => ({...prev, toCompare: undefined}));
         handler.query(targetStream, true)
             .catch((reason) => console.error(reason))
             .finally(() => {
-                setTestHealthToCompare(getTestHealthItems(handler, targetStream, test.nameRef).find(_ => true));
-                setStreams({...streams, toCompare: targetStream});
+                setTestHealthToCompare(getTestHealthItems(handler, targetStream, test.nameRef).at(0));
+                setStreams(prev => ({...prev, toCompare: targetStream}));
                 userStorage.setItem(StoreKey.HistoryToCompare, targetStream);
             });
-    }, [streams, test]);
+    }, [test]);
 
     const onDetailsClick = useCallback((session: TestSessionResult) => {
         if (session.testDataId) {
@@ -589,13 +612,59 @@ export const TestHealthView: React.FC<{ test: TestHealthItem, handler: TestDataH
         }
     }, []);
 
-    const profileOptions: IComboBoxOption[] = profileColors.getDomain().map(p => ({key: p, text: p, styles: {optionText: {backgroundColor: profileColors.get(p), padding: 3, borderRadius: 4, minWidth: 72}}}));
-    const teamOptions: IComboBoxOption[] = teamColors.getDomain().map(p => ({key: p, text: p, styles: {optionText: {backgroundColor: teamColors.get(p), padding: 3, borderRadius: 4, minWidth: 72}}}));
-    const harnessOptions: IComboBoxOption[] = harnessColors.getDomain().map(p => ({key: p, text: p}));
+    const profileOptions: IComboBoxOption[] = profileColors.getDomain().map(p => ({key: p.toLowerCase(), text: p, styles: {optionText: {backgroundColor: profileColors.get(p), padding: 3, borderRadius: 4, minWidth: 72}}}));
+    const teamOptions: IComboBoxOption[] = teamColors.getDomain().map(p => ({key: p.toLowerCase(), text: p, styles: {optionText: {backgroundColor: teamColors.get(p), padding: 3, borderRadius: 4, minWidth: 72}}}));
+    const harnessOptions: IComboBoxOption[] = harnessColors.getDomain().map(p => ({key: p.toLowerCase(), text: p}));
     const notificationOptions: IComboBoxOption[] | undefined = projectStore.streamById(test.stream)?.workflows?.filter(w => !!w.reportChannel).map(w => ({key: w.id, text: w.id, data: w.id}));
     if (!!notificationOptions) {
         notificationOptions.unshift({key: 'disable', text: '- disable notification -', data: undefined});
     }
+    const notificationReportChannel = projectStore.streamById(test.stream)?.workflows?.find(w => w.id === test.notification)?.reportChannel;
+    const getNotificationReportLinkAsync = useCallback(async () => {
+        if (!!notificationReportChannel) {
+            if (!!deepLinkState.notification.url) {
+                return deepLinkState.notification.url;
+            }
+            setDeepLinkState(prev => ({...prev, notification: {isFetching: true}}));
+            let response: DeepLinkResponse | undefined;
+            try {
+                response = await backend.getDeepLinkChannelLink(notificationReportChannel);
+            } catch (error) { /* ignore */ } finally {
+                setDeepLinkState(prev => ({...prev, notification: {isFetching: false, url: response?.url}}));
+                if (!!response?.url) {
+                    return response.url;
+                } else {
+                    errorDialogStore.set({
+                        title: `Error Getting Report Channel Link ${notificationReportChannel}`,
+                        message: "Failed to get a deep link to the report channel."
+                    }, true);
+                }
+            }
+        }
+    }, [notificationReportChannel, deepLinkState.notification.url]);
+    const getDirectMessageLinkAsync = useCallback(async () => {
+        if (!!test.owner || !!test.customers) {
+            if (!!deepLinkState.directMessage.url) {
+                return deepLinkState.directMessage.url;
+            }
+            const userIds = [test.owner, ...(test.customers ?? [])].filter(id => !!id) as string[];
+            setDeepLinkState(prev => ({...prev, directMessage: {isFetching: true}}));
+            let response: DeepLinkResponse | undefined;
+            try {
+                response = await backend.getDeepLinkDirectMessageLink(userIds);
+            } catch (error) { /* ignore */ } finally {
+                setDeepLinkState(prev => ({...prev, directMessage: {isFetching: false, url: response?.url}}));
+                if (!!response?.url) {
+                    return response.url;
+                } else {
+                    errorDialogStore.set({
+                        title: "Error Getting Direct Message Link",
+                        message: "Failed to get a deep link to the direct message channel with the users."
+                    }, true);
+                }
+            }
+        }
+    }, [test.owner, test.customers, deepLinkState.directMessage.url]);
 
     const tags = useMemo(() => test.nameRef?.tagRefs?.keys().map(t => t.name).toArray().sort(), [test]);
 
@@ -642,30 +711,30 @@ export const TestHealthView: React.FC<{ test: TestHealthItem, handler: TestDataH
                                         styles={{input: {backgroundColor: profileColors.get(test.profile)}, root: {backgroundColor: profileColors.get(test.profile), width: 170, height: 24}}}
                                         placeholder="Profile"
                                         allowFreeform
-                                        autoComplete="on"
+                                        autoComplete="off"
                                         options={profileOptions}
                                         text={test.profile}
-                                        onChange={(_1, _2, _3, value) => { test.profile = value; onUpdateItem(test); }}
+                                        onChange={(_1, _2, _3, value) => { const lowerValue = value?.toLowerCase(); test.profile = profileOptions.find(o => o.key === lowerValue)?.text ?? value; onUpdateItem(test); }}
                                     />
                                     <ComboBox
                                         key={`team-${test.updated}`}
                                         styles={{input: {backgroundColor: teamColors.get(test.team)}, root: {backgroundColor: teamColors.get(test.team), width: 170, height: 24}}}
                                         placeholder="Team"
                                         allowFreeform
-                                        autoComplete="on"
+                                        autoComplete="off"
                                         options={teamOptions}
                                         text={test.team}
-                                        onChange={(_1, _2, _3, value) => { test.team = value; onUpdateItem(test); }}
+                                        onChange={(_1, _2, _3, value) => { const lowerValue = value?.toLowerCase(); test.team = teamOptions.find(o => o.key === lowerValue)?.text ?? value; onUpdateItem(test); }}
                                     />
                                     <ComboBox
                                         key={`harness-${test.updated}`}
                                         styles={{root: {width: 170, height: 24}}}
                                         placeholder="Harness"
                                         allowFreeform
-                                        autoComplete="on"
+                                        autoComplete="off"
                                         options={harnessOptions}
                                         text={test.harness}
-                                        onChange={(_1, _2, _3, value) => { test.harness = value; }}
+                                        onChange={(_1, _2, _3, value) => { const lowerValue = value?.toLowerCase(); test.harness = harnessOptions.find(o => o.key === lowerValue)?.text ?? value; }}
                                     />
                                 </Stack>
 
@@ -683,24 +752,32 @@ export const TestHealthView: React.FC<{ test: TestHealthItem, handler: TestDataH
                             </Stack>
                             <Stack tokens={{childrenGap: 8}} style={{width: '100%'}}>
 
-                                <Stack horizontal>
-                                    <Stack horizontal verticalAlign="center" tokens={{childrenGap: 6}} style={{cursor: 'default'}}>
-                                        <Label style={{whiteSpace: 'nowrap'}}>Degradation Notification: </Label>
-                                        {!notificationOptions && <Text style={{color:textEmphasisColor(), cursor: 'default'}}>no workflow available in this stream</Text>}
-                                        {!!notificationOptions && <ComboBox
-                                            key={`notification-${test.updated}`}
-                                            styles={{root: {width: 150, height: 24}}}
-                                            placeholder="none"
-                                            autoComplete="on"
-                                            options={notificationOptions}
-                                            text={test.notification}
-                                            onChange={(_1, option) => { test.notification = option?.data as string | undefined; onUpdateItem(test); }}
-                                            title={test.notification}
-                                        />}
-                                    </Stack>
+                                <Stack horizontal verticalAlign="start" tokens={{childrenGap: 6}} style={{cursor: 'default'}} grow>
+                                    <Label style={{whiteSpace: 'nowrap'}}>Degradation Notification: </Label>
+                                    {!notificationOptions && <Text style={{color:textEmphasisColor(), cursor: 'default'}}>no workflow available in this stream</Text>}
+                                    {!!notificationOptions &&
+                                        <Stack horizontalAlign="end" horizontal>
+                                            <ComboBox
+                                                key={`notification-${test.updated}`}
+                                                styles={{root: {width: 150, height: 24}}}
+                                                placeholder="none"
+                                                autoComplete="on"
+                                                options={notificationOptions}
+                                                text={test.notification}
+                                                onChange={(_1, option) => { test.notification = option?.data as string | undefined; onUpdateItem(test); }}
+                                                title={test.notification}
+                                            />
+                                            {!!notificationReportChannel &&
+                                                <Link style={{padding: '0px 5px'}} title="open target channel" onClick={() => { getNotificationReportLinkAsync().then(link => link && window.open(link, '_blank')); }} disabled={deepLinkState.notification.isFetching}>
+                                                    {!deepLinkState.notification.isFetching && <FontIcon iconName="NavigateExternalInline" style={{fontSize: 17, paddingRight: 4}} />}
+                                                    {!!deepLinkState.notification.isFetching && <Spinner style={{width: 24}} size={SpinnerSize.small} />}
+                                                </Link>
+                                            }
+                                        </Stack>
+                                    }
                                 </Stack>
 
-                                <Stack horizontal tokens={{childrenGap: 12}}>
+                                <Stack horizontal tokens={{childrenGap: 12}} grow>
                                     <Stack tokens={{childrenGap: 6}} style={{borderRadius: 4}} verticalAlign="start">
                                         <Label>Owner: </Label>
                                         <Stack horizontal tokens={{childrenGap: 6}} wrap>
@@ -721,7 +798,7 @@ export const TestHealthView: React.FC<{ test: TestHealthItem, handler: TestDataH
                                             }
                                         </Stack>
                                     </Stack>
-                                    <Stack tokens={{childrenGap: 6}} style={{borderRadius: 4}} verticalAlign="start">
+                                    <Stack tokens={{childrenGap: 6}} style={{borderRadius: 4}} verticalAlign="start" grow>
                                         <Label>Customers: </Label>
                                         <Stack horizontal tokens={{childrenGap: 6}} wrap>
                                             {!!test.customers &&
@@ -739,6 +816,12 @@ export const TestHealthView: React.FC<{ test: TestHealthItem, handler: TestDataH
                                             </Stack>
                                         </Stack>
                                     </Stack>
+                                    {(!!test.owner || !!test.customers) &&
+                                        <Link onClick={() => { getDirectMessageLinkAsync().then(link => link && window.open(link, '_blank')); }} title="start chat with owner and customers" disabled={deepLinkState.directMessage.isFetching}>
+                                            {!deepLinkState.directMessage.isFetching && <FontIcon iconName="OfficeChat" style={{cursor: 'pointer', fontSize: 17, padding: 2}}/>}
+                                            {!!deepLinkState.directMessage.isFetching && <Spinner style={{width: 24}} size={SpinnerSize.small} />}
+                                        </Link>
+                                    }
                                 </Stack>
 
                             </Stack>
@@ -789,8 +872,8 @@ export const TestHealthView: React.FC<{ test: TestHealthItem, handler: TestDataH
                                     </Stack>
                                 }
                                 {!test.user && <Text style={{color: textEmphasisColor()}} variant="small">no user update</Text>}
-                                <Link href={`/test-automation/log/${test.nameRef.id}`} target="_blank">
-                                    <FontIcon iconName="TimeEntry" style={{cursor: 'pointer', fontSize: 17, padding: 2}} title="open audit log"/>
+                                <Link href={`/test-automation/log/${test.nameRef.id}`} target="_blank" title="open audit history">
+                                    <FontIcon iconName="TimeEntry" style={{cursor: 'pointer', fontSize: 17, padding: 2}}/>
                                 </Link>
                             </Stack>
                         </Stack>
@@ -988,10 +1071,7 @@ export const TestAutomationHealth: React.FC<{ handler: TestDataHandler}> = obser
 
     const forceListRefresh = () => {
         updateColorDomains(healthListState.items);
-        setHealthListState({
-            ...healthListState,
-            items: healthListState.items.slice()
-        });
+        setHealthListState(prev => ({ ...prev, items: healthListState.items.slice() }));
     }
 
     // initiate data on filtering change
@@ -1019,7 +1099,7 @@ export const TestAutomationHealth: React.FC<{ handler: TestDataHandler}> = obser
 
             setHealthDetailsState({test: testHealth});
         }
-    }, [handler.searchUpdated, healthDetailsState]);
+    }, [handler.searchUpdated, testId, healthListState.items]);
 
     return  <Stack grow style={{position: "relative"}}>
                 {!!healthDetailsState.test &&
diff --git a/Engine/Source/Programs/Horde/HordeDashboard/src/backend/Api.ts b/Engine/Source/Programs/Horde/HordeDashboard/src/backend/Api.ts
--- a/Engine/Source/Programs/Horde/HordeDashboard/src/backend/Api.ts
+++ b/Engine/Source/Programs/Horde/HordeDashboard/src/backend/Api.ts
@@ -5502,4 +5502,10 @@ export type GetBuildHealthFilterResponse = {
 	updateTimeUtc: Date | string;
 }
 
-// #endregion -- Build Health Filters --
\ No newline at end of file
+// #endregion -- Build Health Filters --
+
+/** Deep Link Response */
+export type DeepLinkResponse = {
+	/// The deep link url
+	url?: string;
+}
\ No newline at end of file
diff --git a/Engine/Source/Programs/Horde/HordeDashboard/src/backend/Backend.ts b/Engine/Source/Programs/Horde/HordeDashboard/src/backend/Backend.ts
--- a/Engine/Source/Programs/Horde/HordeDashboard/src/backend/Backend.ts
+++ b/Engine/Source/Programs/Horde/HordeDashboard/src/backend/Backend.ts
@@ -1,7 +1,7 @@
 // Copyright Epic Games, Inc. All Rights Reserved.
 
 import templateCache from '../backend/TemplateCache';
-import { AccountClaimMessage, AddBuildHealthFiltersRequest, AgentData, AgentQuery, ApproveAgentsRequest, AuditLogEntry, AuditLogQuery, BatchUpdatePoolRequest, ChangeSummaryData, CreateAccountRequest, CreateBisectTaskRequest, CreateBisectTaskResponse, CreateDeviceRequest, CreateDeviceResponse, CreateExternalIssueRequest, CreateExternalIssueResponse, CreateJobRequest, CreateJobResponse, CreateNoticeRequest, CreatePoolRequest, CreateServiceAccountRequest, CreateServiceAccountResponse, CreateSoftwareResponse, CreateSubscriptionRequest, CreateSubscriptionResponse, CreateZipRequest, DashboardPreference, DevicePoolTelemetryQuery, DeviceTelemetryQuery, EventData, FindArtifactsResponse, FindIssueResponse, FindJobTimingsResponse, GetAccountResponse, GetAgentLeaseResponse, GetAgentTelemetrySampleResponse, GetArtifactDirectoryResponse, GetArtifactResponse, GetBisectTaskResponse, GetBuildHealthFilterResponse, GetBuildHealthFiltersRequest, GetDashboardConfigResponse, GetDevicePlatformResponse, GetDevicePoolResponse, GetDevicePoolTelemetryResponse, GetDeviceReservationResponse, GetDeviceResponse, GetDeviceTelemetryResponse, GetExternalIssueProjectResponse, GetExternalIssueResponse, GetIssueStreamResponse, GetJobsTabResponse, GetJobStepRefResponse, GetJobStepTraceResponse, GetJobTimingResponse, GetLogEventResponse, GetLogFileResponse, GetNoticeResponse, GetNotificationResponse, GetPendingAgentsResponse, GetPerforceServerStatusResponse, GetPoolResponse, GetPoolSummaryResponse, GetServerInfoResponse, GetServerSettingsResponse, GetServiceAccountResponse, GetSoftwareResponse, GetSubscriptionResponse, GetTestDataDetailsResponse, GetTestDataRefResponse, GetTestMetaResponse, GetTestResponse, GetTestsRequest, GetTestStreamResponse, GetToolResponse, GetToolSummaryResponse, GetUserResponse, GetUtilizationTelemetryResponse, GlobalConfig, IssueData, IssueQuery, IssueQueryV2, JobData, JobQuery, JobsTabColumnType, JobStepOutcome, JobStreamQuery, JobTimingsQuery, LeaseData, LogData, LogEventQuery, LogLineData, PoolData, PoolQuery, PreflightConfigResponse, ProjectData, ScheduleData, ScheduleQuery, SearchLogFileResponse, ServerPluginInfoResponse, ServerStatusResponse, ServerUpdateResponse, SessionData, StreamData, TabType, TestData, UpdateAccountRequest, UpdateAgentRequest, UpdateBisectTaskRequest, UpdateBuildHealthFiltersRequest, UpdateCurrentAccountRequest, UpdateDeploymentRequest, UpdateDeviceRequest, UpdateGlobalConfigRequest, UpdateIssueRequest, UpdateJobRequest, UpdateLeaseRequest, UpdateNoticeRequest, UpdateNotificationsRequest, UpdatePoolRequest, UpdateServerSettingsRequest, UpdateServiceAccountRequest, UpdateServiceAccountResponse, UpdateStepRequest, UpdateStepResponse, UpdateTemplateRefRequest, UpdateUserRequest, UsersQuery } from './Api';
+import { AccountClaimMessage, AddBuildHealthFiltersRequest, AgentData, AgentQuery, ApproveAgentsRequest, AuditLogEntry, AuditLogQuery, BatchUpdatePoolRequest, ChangeSummaryData, CreateAccountRequest, CreateBisectTaskRequest, CreateBisectTaskResponse, CreateDeviceRequest, CreateDeviceResponse, CreateExternalIssueRequest, CreateExternalIssueResponse, CreateJobRequest, CreateJobResponse, CreateNoticeRequest, CreatePoolRequest, CreateServiceAccountRequest, CreateServiceAccountResponse, CreateSoftwareResponse, CreateSubscriptionRequest, CreateSubscriptionResponse, CreateZipRequest, DashboardPreference, DeepLinkResponse, DevicePoolTelemetryQuery, DeviceTelemetryQuery, EventData, FindArtifactsResponse, FindIssueResponse, FindJobTimingsResponse, GetAccountResponse, GetAgentLeaseResponse, GetAgentTelemetrySampleResponse, GetArtifactDirectoryResponse, GetArtifactResponse, GetBisectTaskResponse, GetBuildHealthFilterResponse, GetBuildHealthFiltersRequest, GetDashboardConfigResponse, GetDevicePlatformResponse, GetDevicePoolResponse, GetDevicePoolTelemetryResponse, GetDeviceReservationResponse, GetDeviceResponse, GetDeviceTelemetryResponse, GetExternalIssueProjectResponse, GetExternalIssueResponse, GetIssueStreamResponse, GetJobsTabResponse, GetJobStepRefResponse, GetJobStepTraceResponse, GetJobTimingResponse, GetLogEventResponse, GetLogFileResponse, GetNoticeResponse, GetNotificationResponse, GetPendingAgentsResponse, GetPerforceServerStatusResponse, GetPoolResponse, GetPoolSummaryResponse, GetServerInfoResponse, GetServerSettingsResponse, GetServiceAccountResponse, GetSoftwareResponse, GetSubscriptionResponse, GetTestDataDetailsResponse, GetTestDataRefResponse, GetTestMetaResponse, GetTestResponse, GetTestsRequest, GetTestStreamResponse, GetToolResponse, GetToolSummaryResponse, GetUserResponse, GetUtilizationTelemetryResponse, GlobalConfig, IssueData, IssueQuery, IssueQueryV2, JobData, JobQuery, JobsTabColumnType, JobStepOutcome, JobStreamQuery, JobTimingsQuery, LeaseData, LogData, LogEventQuery, LogLineData, PoolData, PoolQuery, PreflightConfigResponse, ProjectData, ScheduleData, ScheduleQuery, SearchLogFileResponse, ServerPluginInfoResponse, ServerStatusResponse, ServerUpdateResponse, SessionData, StreamData, TabType, TestData, UpdateAccountRequest, UpdateAgentRequest, UpdateBisectTaskRequest, UpdateBuildHealthFiltersRequest, UpdateCurrentAccountRequest, UpdateDeploymentRequest, UpdateDeviceRequest, UpdateGlobalConfigRequest, UpdateIssueRequest, UpdateJobRequest, UpdateLeaseRequest, UpdateNoticeRequest, UpdateNotificationsRequest, UpdatePoolRequest, UpdateServerSettingsRequest, UpdateServiceAccountRequest, UpdateServiceAccountResponse, UpdateStepRequest, UpdateStepResponse, UpdateTemplateRefRequest, UpdateUserRequest, UsersQuery } from './Api';
 import dashboard, { Dashboard } from './Dashboard';
 import { ChallengeStatus, Fetch } from './Fetch';
 import { projectStore } from './ProjectStore';
@@ -1952,6 +1952,35 @@ export class Backend {
 
     // #endregion -- Build Health Filters --
 
+    // #region -- Deep Links --
+
+    /// Get a deep link to a channel
+    getDeepLinkChannelLink(channelId: string): Promise<DeepLinkResponse> {
+        return new Promise<DeepLinkResponse>((resolve, reject) => {
+            this._fetch.get(`/api/v1/deeplink/channel`, {
+                params: {channelId: channelId}
+            }).then((value) => {
+                resolve(value.data as DeepLinkResponse);
+            }).catch((reason) => {
+                reject(reason);
+            });
+        });
+    }
+    
+    /// Get a deep link to a direct message
+    getDeepLinkDirectMessageLink(userIds: string[]): Promise<DeepLinkResponse> {
+        const request = { userIds: userIds };
+        return new Promise<DeepLinkResponse>((resolve, reject) => {
+            this._fetch.post(`/api/v1/deeplink/directmessage`, request).then((value) => {
+                resolve(value.data as DeepLinkResponse);
+            }).catch((reason) => {
+                reject(reason);
+            });
+        });
+    }
+
+    // #endregion -- Deep Links --
+
     private async update() {
 
         if (this.updateID === "updating") {
