[Horde Dashboard] Added raw `agent hours` metric to pool utilization report view and functionality to click bar chart segments for a fixed overview popup

#Synced-CL 48478172

diff --git a/Engine/Source/Programs/Horde/HordeDashboard/src/components/UtilizationReportView.tsx b/Engine/Source/Programs/Horde/HordeDashboard/src/components/UtilizationReportView.tsx
--- a/Engine/Source/Programs/Horde/HordeDashboard/src/components/UtilizationReportView.tsx
+++ b/Engine/Source/Programs/Horde/HordeDashboard/src/components/UtilizationReportView.tsx
@@ -1,6 +1,6 @@
 // Copyright Epic Games, Inc. All Rights Reserved.
 
-import { DatePicker, Dropdown, DropdownMenuItemType, IDatePickerStrings, IDropdownOption, IDropdownStyles, ISliderStyles, mergeStyleSets, Slider, Stack } from '@fluentui/react';
+import { DatePicker, Dropdown, DropdownMenuItemType, IDatePickerStrings, IDropdownOption, IDropdownStyles, ISliderStyles, mergeStyleSets, Slider, Stack, Callout, DirectionalHint } from '@fluentui/react';
 import { IVerticalStackedChartProps, IVSChartDataPoint, VerticalStackedBarChart } from '@fluentui/react-charting';
 import { observer } from 'mobx-react-lite';
 import moment from 'moment-timezone';
@@ -49,6 +49,15 @@ type UpdateStateOptions = {
    selectedFormat?: number | null | undefined;
 }
 
+type PinnedCalloutState = {
+   legend: string;
+   data: number;
+   color: string;
+   yAxisCalloutData: string;
+   x: number,
+   y: number
+}
+
 const intervalToOptions: { [key: number]: IntervalOptions } = {
    1: {
       bucketInterval: 1,
@@ -174,6 +183,8 @@ export const UtilizationReportView: React.FC = observer(() => {
       noData: false
    });
 
+   const [pinnedCallout, setPinnedCallout] = useState<PinnedCalloutState | null>(null);
+
    const dropdownStyles: Partial<IDropdownStyles> = { dropdown: { width: 300 } };
    const intervalStyle: Partial<ISliderStyles> = { root: { width: 200, marginRight: -40 }, titleLabel: { position: 'relative', top: 5 }, container: { position: 'relative', top: 13 } };
    const rangeStyle: Partial<ISliderStyles> = { root: { width: 200, marginRight: -15 }, titleLabel: { position: 'relative', top: 5 }, container: { position: 'relative', top: 13 } };
@@ -354,6 +365,8 @@ export const UtilizationReportView: React.FC = observer(() => {
       newState.updating = false;
       newState.yMaxValue = 0;
 
+      setPinnedCallout(null);
+
       let dayIdx = moment();
 
       // get start and end date
@@ -512,11 +525,12 @@ export const UtilizationReportView: React.FC = observer(() => {
             else {
                for (let streamKey in streamTotalTime) {
                   if (streamTotalTime[streamKey] !== 0) {
+                     const streamTime = streamTotalTime[streamKey];
                      chartPoints.push({
                         legend: streamKey,
-                        data: streamTotalTime[streamKey],
+                        data: streamTime,
                         color: getPoolColor(streamKey),
-                        yAxisCalloutData: `${(streamTotalTime[streamKey] / totalTime).toLocaleString(undefined, { style: 'percent', maximumFractionDigits: 1 })}`
+                        yAxisCalloutData: `${(streamTime / totalTime).toLocaleString(undefined, { style: 'percent', maximumFractionDigits: 1 })} (${streamTime.toFixed(1)} agent hours)`
                      });
                      totalTakenTime += streamTotalTime[streamKey];
                   }
@@ -791,7 +805,38 @@ export const UtilizationReportView: React.FC = observer(() => {
                                        barWidth={intervalToOptions[optionsState.selectedFormat].barWidth / optionsState.selectedRange}
                                        hideLegend={true}
                                        yAxisTickFormat={yAxisFormatter}
+                                       onBarClick={(event, data: any) => {
+                                          if (pinnedCallout?.legend === data.legend) {
+                                             setPinnedCallout(null);
+                                          } else {
+                                             setPinnedCallout({
+                                                legend: data.legend,
+                                                data: data.data,
+                                                color: data.color,
+                                                yAxisCalloutData: data.yAxisCalloutData,
+                                                x: event.clientX,
+                                                y: event.clientY
+                                             });
+                                          }
+                                       }}
                                     />
+                                    {pinnedCallout && (
+                                       <Callout
+                                          target={{ x: pinnedCallout.x, y: pinnedCallout.y }}
+                                          onDismiss={() => setPinnedCallout(null)}
+                                          directionalHint={DirectionalHint.bottomLeftEdge}
+                                          isBeakVisible={true}
+                                          gapSpace={10}
+                                          setInitialFocus
+                                       >
+                                          <div style={{ padding: 12, minWidth: 150 }}>
+                                             <div style={{ fontWeight: 600, color: pinnedCallout.color, marginBottom: 4 }}>
+                                                {pinnedCallout.legend}
+                                             </div>
+                                             <div>{pinnedCallout.yAxisCalloutData}</div>
+                                          </div>
+                                       </Callout>
+                                    )}
                                  </div>
                                     :
                                     <div>Fetching Data...</div>
