Horde: Updating schema documentation

#rnx

diff --git a/Engine/Source/Programs/Horde/Docs/Config/Schema/AclActions.md b/Engine/Source/Programs/Horde/Docs/Config/Schema/AclActions.md
--- a/Engine/Source/Programs/Horde/Docs/Config/Schema/AclActions.md
+++ b/Engine/Source/Programs/Horde/Docs/Config/Schema/AclActions.md
@@ -145,17 +145,6 @@
 | `DeleteAccount` |  |
 | `ViewAccount` |  |
 
-## Agents
-
-| Name | Description |
-| ---- | ----------- |
-| `CreateAgent` | Ability to create an agent. This may be done explicitly, or granted to agents to allow them to self-register. |
-| `CreateWorkstationAgent` | Ability to create a workstation agent. Granted to users with hardware to offer. |
-| `UpdateAgent` | Update an agent's name, pools, etc... |
-| `DeleteAgent` | Soft-delete an agent |
-| `ViewAgent` | View an agent |
-| `ListAgents` | List the available agents |
-
 ## Artifacts
 
 | Name | Description |
@@ -180,15 +169,6 @@
 | ---- | ----------- |
 | `AddUpdateDeleteBuildHealthFilter` |  |
 
-## Compute
-
-| Name | Description |
-| ---- | ----------- |
-| `AddComputeTasks` | User can add tasks to the compute cluster |
-| `GetComputeTasks` | User can get and list tasks from the compute cluster |
-| `UbaCacheRead` | User can read from Unreal Build Accelerator cache |
-| `UbaCacheWrite` | User can write to Unreal Build Accelerator cache |
-
 ## Devices
 
 | Name | Description |
@@ -207,24 +187,6 @@
 | `RetryJobStep` |  |
 | `ViewJob` |  |
 
-## Leases
-
-| Name | Description |
-| ---- | ----------- |
-| `ViewLeases` | View all the leases that an agent has worked on |
-| `ViewLeaseTasks` | View the task data for a lease |
-
-## Logs
-
-| Name | Description |
-| ---- | ----------- |
-| `CreateLog` | Ability to create a log. Implicitly granted to agents. |
-| `UpdateLog` | Ability to update log metadata |
-| `ViewLog` | Ability to view a log contents |
-| `WriteLogData` | Ability to write log data |
-| `CreateEvent` | Ability to create events |
-| `ViewEvent` | Ability to view events |
-
 ## Notices
 
 | Name | Description |
@@ -239,16 +201,6 @@
 | ---- | ----------- |
 | `CreateSubscription` |  |
 
-## Pools
-
-| Name | Description |
-| ---- | ----------- |
-| `CreatePool` | Create a global pool of agents |
-| `UpdatePool` | Modify an agent pool |
-| `DeletePool` | Delete an agent pool |
-| `ViewPool` | Ability to view a pool |
-| `ListPools` | View all the available agent pools |
-
 ## Projects
 
 | Name | Description |
@@ -274,21 +226,6 @@
 | `DeleteAccount` |  |
 | `ViewAccount` |  |
 
-## Sessions
-
-| Name | Description |
-| ---- | ----------- |
-| `CreateSession` | Granted to agents to call CreateSession, which returns a bearer token identifying themselves valid to call UpdateSesssion via gRPC. |
-| `ViewSession` | Allows viewing information about an agent session |
-
-## Software
-
-| Name | Description |
-| ---- | ----------- |
-| `UploadSoftware` | Ability to upload new versions of the agent software |
-| `DownloadSoftware` | Ability to download the agent software |
-| `DeleteSoftware` | Ability to delete agent software |
-
 ## Streams
 
 | Name | Description |
@@ -320,12 +257,12 @@
 
 | Name | Description |
 | ---- | ----------- |
-| `CreateAgent` |  |
-| `CreateWorkstationAgent` |  |
-| `UpdateAgent` |  |
-| `DeleteAgent` |  |
-| `ViewAgent` |  |
-| `ListAgents` |  |
+| `CreateAgent` | Ability to create an agent. This may be done explicitly, or granted to agents to allow them to self-register. |
+| `CreateWorkstationAgent` | Ability to create a workstation agent. Granted to users with hardware to offer. |
+| `UpdateAgent` | Update an agent's name, pools, etc... |
+| `DeleteAgent` | Soft-delete an agent |
+| `ViewAgent` | View an agent |
+| `ListAgents` | List the available agents |
 
 ## Artifacts
 
@@ -355,10 +292,10 @@
 
 | Name | Description |
 | ---- | ----------- |
-| `AddComputeTasks` |  |
-| `GetComputeTasks` |  |
-| `UbaCacheRead` |  |
-| `UbaCacheWrite` |  |
+| `AddComputeTasks` | User can add tasks to the compute cluster |
+| `GetComputeTasks` | User can get and list tasks from the compute cluster |
+| `UbaCacheRead` | User can read from Unreal Build Accelerator cache |
+| `UbaCacheWrite` | User can write to Unreal Build Accelerator cache |
 
 ## Devices
 
@@ -382,19 +319,19 @@
 
 | Name | Description |
 | ---- | ----------- |
-| `ViewLeases` |  |
-| `ViewLeaseTasks` |  |
+| `ViewLeases` | View all the leases that an agent has worked on |
+| `ViewLeaseTasks` | View the task data for a lease |
 
 ## Logs
 
 | Name | Description |
 | ---- | ----------- |
-| `CreateLog` |  |
-| `UpdateLog` |  |
-| `ViewLog` |  |
-| `WriteLogData` |  |
-| `CreateEvent` |  |
-| `ViewEvent` |  |
+| `CreateLog` | Ability to create a log. Implicitly granted to agents. |
+| `UpdateLog` | Ability to update log metadata |
+| `ViewLog` | Ability to view a log contents |
+| `WriteLogData` | Ability to write log data |
+| `CreateEvent` | Ability to create events |
+| `ViewEvent` | Ability to view events |
 
 ## Notices
 
@@ -414,11 +351,11 @@
 
 | Name | Description |
 | ---- | ----------- |
-| `CreatePool` |  |
-| `UpdatePool` |  |
-| `DeletePool` |  |
-| `ViewPool` |  |
-| `ListPools` |  |
+| `CreatePool` | Create a global pool of agents |
+| `UpdatePool` | Modify an agent pool |
+| `DeletePool` | Delete an agent pool |
+| `ViewPool` | Ability to view a pool |
+| `ListPools` | View all the available agent pools |
 
 ## Projects
 
@@ -449,16 +386,16 @@
 
 | Name | Description |
 | ---- | ----------- |
-| `CreateSession` |  |
-| `ViewSession` |  |
+| `CreateSession` | Granted to agents to call CreateSession, which returns a bearer token identifying themselves valid to call UpdateSesssion via gRPC. |
+| `ViewSession` | Allows viewing information about an agent session |
 
 ## Software
 
 | Name | Description |
 | ---- | ----------- |
-| `UploadSoftware` |  |
-| `DownloadSoftware` |  |
-| `DeleteSoftware` |  |
+| `UploadSoftware` | Ability to upload new versions of the agent software |
+| `DownloadSoftware` | Ability to download the agent software |
+| `DeleteSoftware` | Ability to delete agent software |
 
 ## Streams
 
@@ -778,6 +715,12 @@
 | `UpdateReplicator` |  |
 | `ViewReplicator` |  |
 
+## Secrets
+
+| Name | Description |
+| ---- | ----------- |
+| `ViewSecret` | View a credential |
+
 ## ServiceAccounts
 
 | Name | Description |
@@ -949,6 +892,12 @@
 | `UpdateReplicator` |  |
 | `ViewReplicator` |  |
 
+## Secrets
+
+| Name | Description |
+| ---- | ----------- |
+| `ViewSecret` |  |
+
 ## ServiceAccounts
 
 | Name | Description |
@@ -991,13 +940,6 @@
 | `TestHealthRead` |  |
 | `TestHealthWrite` |  |
 
-## Tools
-
-| Name | Description |
-| ---- | ----------- |
-| `DownloadTool` | Ability to download a tool |
-| `UploadTool` | Ability to upload new tool versions |
-
 ## Accounts
 
 | Name | Description |
@@ -1127,6 +1069,12 @@
 | `UpdateReplicator` |  |
 | `ViewReplicator` |  |
 
+## Secrets
+
+| Name | Description |
+| ---- | ----------- |
+| `ViewSecret` |  |
+
 ## ServiceAccounts
 
 | Name | Description |
@@ -1151,6 +1099,19 @@
 | `DownloadSoftware` |  |
 | `DeleteSoftware` |  |
 
+## Storage
+
+| Name | Description |
+| ---- | ----------- |
+| `ReadBlobs` | Ability to read blobs from the storage service |
+| `WriteBlobs` | Ability to write blobs to the storage service |
+| `ReadAliases` | Ability to read aliases from the storage service |
+| `WriteAliases` | Ability to write aliases to the storage service |
+| `DeleteAliases` | Ability to write aliases to the storage service |
+| `ReadRefs` | Ability to read refs from the storage service |
+| `WriteRefs` | Ability to write refs to the storage service |
+| `DeleteRefs` | Ability to delete refs |
+
 ## Streams
 
 | Name | Description |
@@ -1162,12 +1123,6 @@
 | `ViewChanges` |  |
 | `ViewTemplate` |  |
 
-## Symbols
-
-| Name | Description |
-| ---- | ----------- |
-| `ReadSymbols` | Ability to download symbols |
-
 ## TestData
 
 | Name | Description |
@@ -1175,13 +1130,6 @@
 | `TestHealthRead` |  |
 | `TestHealthWrite` |  |
 
-## Tools
-
-| Name | Description |
-| ---- | ----------- |
-| `DownloadTool` |  |
-| `UploadTool` |  |
-
 ## Accounts
 
 | Name | Description |
@@ -1315,7 +1263,7 @@
 
 | Name | Description |
 | ---- | ----------- |
-| `ViewSecret` | View a credential |
+| `ViewSecret` |  |
 
 ## ServiceAccounts
 
@@ -1341,6 +1289,19 @@
 | `DownloadSoftware` |  |
 | `DeleteSoftware` |  |
 
+## Storage
+
+| Name | Description |
+| ---- | ----------- |
+| `ReadBlobs` |  |
+| `WriteBlobs` |  |
+| `ReadAliases` |  |
+| `WriteAliases` |  |
+| `DeleteAliases` |  |
+| `ReadRefs` |  |
+| `WriteRefs` |  |
+| `DeleteRefs` |  |
+
 ## Streams
 
 | Name | Description |
@@ -1352,12 +1313,6 @@
 | `ViewChanges` |  |
 | `ViewTemplate` |  |
 
-## Symbols
-
-| Name | Description |
-| ---- | ----------- |
-| `ReadSymbols` |  |
-
 ## TestData
 
 | Name | Description |
@@ -1369,8 +1324,8 @@
 
 | Name | Description |
 | ---- | ----------- |
-| `DownloadTool` |  |
-| `UploadTool` |  |
+| `DownloadTool` | Ability to download a tool |
+| `UploadTool` | Ability to upload new tool versions |
 
 ## Accounts
 
@@ -1535,14 +1490,14 @@
 
 | Name | Description |
 | ---- | ----------- |
-| `ReadBlobs` | Ability to read blobs from the storage service |
-| `WriteBlobs` | Ability to write blobs to the storage service |
-| `ReadAliases` | Ability to read aliases from the storage service |
-| `WriteAliases` | Ability to write aliases to the storage service |
-| `DeleteAliases` | Ability to write aliases to the storage service |
-| `ReadRefs` | Ability to read refs from the storage service |
-| `WriteRefs` | Ability to write refs to the storage service |
-| `DeleteRefs` | Ability to delete refs |
+| `ReadBlobs` |  |
+| `WriteBlobs` |  |
+| `ReadAliases` |  |
+| `WriteAliases` |  |
+| `DeleteAliases` |  |
+| `ReadRefs` |  |
+| `WriteRefs` |  |
+| `DeleteRefs` |  |
 
 ## Streams
 
@@ -1559,7 +1514,7 @@
 
 | Name | Description |
 | ---- | ----------- |
-| `ReadSymbols` |  |
+| `ReadSymbols` | Ability to download symbols |
 
 ## TestData
 
