Horde: Update the version of dotnet used to build to 10.0.

#Synced-CL 50766644

diff --git a/Engine/Source/Programs/Horde/BuildHorde.xml b/Engine/Source/Programs/Horde/BuildHorde.xml
--- a/Engine/Source/Programs/Horde/BuildHorde.xml
+++ b/Engine/Source/Programs/Horde/BuildHorde.xml
@@ -28,7 +28,7 @@
 	<Property Name="InformationalVersion" Value="$(InformationalVersion)-PF-$(PreflightChange)" If="'$(PreflightChange)' != ''"/>
 	<Property Name="VersionArguments" Value="/p:Version=$(Version).0 /p:InformationalVersion=$(InformationalVersion)"/>
 
-	<Property Name="DotNetRootPath" Value="$(RootDir)/Engine/Binaries/ThirdParty/DotNet/8.0.412"/>
+	<Property Name="DotNetRootPath" Value="$(RootDir)/Engine/Binaries/ThirdParty/DotNet/10.0"/>
 	<Property Name="DotNetPath" Value=""/>
 	<Property Name="DotNetPath" Value="$(DotNetRootPath)/win-x64/dotnet.exe" If="$(HostPlatform) == 'Win64' And $(HostArchitecture) == 'X64'"/>
 	<Property Name="DotNetPath" Value="$(DotNetRootPath)/win-arm64/dotnet.exe" If="$(HostPlatform) == 'Win64' And $(HostArchitecture) == 'Arm64'"/>
