Horde: Adding more support for ICommit interface

diff --git a/Engine/Source/Programs/Horde/Drivers/JobDriver/Execution/JobExecutor.cs b/Engine/Source/Programs/Horde/Drivers/JobDriver/Execution/JobExecutor.cs
--- a/Engine/Source/Programs/Horde/Drivers/JobDriver/Execution/JobExecutor.cs
+++ b/Engine/Source/Programs/Horde/Drivers/JobDriver/Execution/JobExecutor.cs
@@ -28,7 +28,21 @@ using StatusCode = Grpc.Core.StatusCode;
 
 namespace JobDriver.Execution
 {
-	record class JobStepInfo
+	/// <summary>
+	/// Information about a job step to be executed
+	/// </summary>
+	/// <param name="StepId">Unique identifier for the step</param>
+	/// <param name="LogId">Log identifier for the step output</param>
+	/// <param name="Name">Name of the step</param>
+	/// <param name="Credentials">Credentials available to the step</param>
+	/// <param name="Properties">Properties for the step</param>
+	/// <param name="EnvVars">Environment variables for the step</param>
+	/// <param name="Warnings">Whether warnings are enabled</param>
+	/// <param name="Inputs">Input dependencies for the step</param>
+	/// <param name="OutputNames">Names of outputs produced by the step</param>
+	/// <param name="PublishOutputs">Indices of outputs to publish</param>
+	/// <param name="Artifacts">Artifact definitions for the step</param>
+	public record class JobStepInfo
 	(
 		JobStepId StepId,
 		LogId LogId,
@@ -43,22 +57,58 @@ namespace JobDriver.Execution
 		IReadOnlyList<RpcCreateGraphArtifactRequest> Artifacts
 	)
 	{
+		/// <summary>
+		/// Constructor from an RPC response
+		/// </summary>
 		public JobStepInfo(RpcBeginStepResponse response)
 			: this(JobStepId.Parse(response.StepId), LogId.Parse(response.LogId), response.Name, response.Credentials, response.Properties, response.EnvVars, response.Warnings, response.Inputs, response.OutputNames, response.PublishOutputs, response.Artifacts)
 		{
 		}
 	}
 
-	class JobExecutorOptions
+	/// <summary>
+	/// Options for creating a job executor
+	/// </summary>
+	public class JobExecutorOptions
 	{
+		/// <summary>
+		/// The Horde client instance
+		/// </summary>
 		public IHordeClient HordeClient { get; }
+
+		/// <summary>
+		/// Working directory for the executor
+		/// </summary>
 		public DirectoryReference WorkingDir { get; }
+
+		/// <summary>
+		/// Processes to terminate before/after execution
+		/// </summary>
 		public IReadOnlyList<ProcessToTerminate>? ProcessesToTerminate { get; }
+
+		/// <summary>
+		/// Identifier for the job
+		/// </summary>
 		public JobId JobId { get; }
+
+		/// <summary>
+		/// Identifier for the batch within the job
+		/// </summary>
 		public JobStepBatchId BatchId { get; }
+
+		/// <summary>
+		/// Identifier for the lease
+		/// </summary>
 		public LeaseId LeaseId { get; }
+
+		/// <summary>
+		/// Options for the job
+		/// </summary>
 		public RpcJobOptions JobOptions { get; }
 
+		/// <summary>
+		/// Constructor
+		/// </summary>
 		public JobExecutorOptions(IHordeClient hordeClient, DirectoryReference workingDir, IReadOnlyList<ProcessToTerminate>? processesToTerminate, JobId jobId, JobStepBatchId batchId, LeaseId leaseId, RpcJobOptions jobOptions)
 		{
 			HordeClient = hordeClient;
@@ -71,8 +121,12 @@ namespace JobDriver.Execution
 		}
 	}
 
-	abstract class JobExecutor : IDisposable
+	/// <summary>
+	/// Abstract base class for job executors. Handles running BuildGraph scripts on an agent.
+	/// </summary>
+	public abstract class JobExecutor : IDisposable
 	{
+#pragma warning disable CS1591
 		protected class ExportedNode
 		{
 			public string Name { get; set; } = String.Empty;
@@ -137,6 +191,7 @@ namespace JobDriver.Execution
 			public List<ExportedBadge> Badges { get; set; } = new List<ExportedBadge>();
 			public List<ExportedArtifact> Artifacts { get; set; } = new List<ExportedArtifact>();
 		}
+#pragma warning restore CS1591
 
 		class TraceEvent
 		{
@@ -205,40 +260,36 @@ namespace JobDriver.Execution
 		string PreprocessedScript => $"{EnginePath}/Saved/Horde/Preprocessed.xml";
 		string PreprocessedSchema => $"{EnginePath}/Saved/Horde/Preprocessed.xsd";
 
-		protected List<string> _targets = new List<string>();
-		protected string? _scriptFileName;
-		protected bool _preprocessScript;
-		protected bool _savePreprocessedScript;
-
+#pragma warning disable CS1591
+		protected List<string> Targets { get; set; } = new List<string>();
+		protected string? ScriptFileName { get; set; }
+		protected bool PreprocessScript { get; set; }
+		protected bool SavePreprocessedScript { get; set; }
 		protected string EnginePath => String.IsNullOrEmpty(Batch.EnginePath)? "Engine" : Batch.EnginePath;
-
 		protected Tracer Tracer { get; }
+#pragma warning restore CS1591
 
 		/// <summary>
 		/// Logger for the local agent process (as opposed to job logger)
 		/// </summary>
 		protected ILogger Logger { get; }
 
+#pragma warning disable CS1591
 		protected IHordeClient HordeClient { get; }
 		protected DirectoryReference WorkingDir { get; }
 		protected JobId JobId { get; }
 		protected JobStepBatchId BatchId { get; }
 		protected LeaseId LeaseId { get; }
 		protected IReadOnlyList<ProcessToTerminate>? ProcessesToTerminate { get; }
-
 		protected RpcBeginBatchResponse Batch { get; private set; }
-
-		protected List<string> _additionalArguments = new List<string>();
-		protected bool _allowTargetChanges;
-
-		protected bool _compileAutomationTool = true;
-
+		protected List<string> AdditionalArguments { get; set; } = new List<string>();
+		protected bool AllowTargetChanges { get; set; }
+		protected bool CompileAutomationTool { get; set; } = true;
 		protected RpcJobOptions JobOptions { get; }
-
 		protected JobRpc.JobRpcClient JobRpc { get; private set; }
-		protected Dictionary<string, string> _remapAgentTypes = new Dictionary<string, string>(StringComparer.OrdinalIgnoreCase);
-
-		protected Dictionary<string, string> _envVars = new Dictionary<string, string>(StringComparer.OrdinalIgnoreCase);
+		protected Dictionary<string, string> RemapAgentTypes { get; set; } = new Dictionary<string, string>(StringComparer.OrdinalIgnoreCase);
+		protected Dictionary<string, string> EnvVars { get; set; } = new Dictionary<string, string>(StringComparer.OrdinalIgnoreCase);
+#pragma warning restore CS1591
 
 		/// <summary>
 		/// How often to poll the server checking if a step has been aborted
@@ -276,6 +327,9 @@ namespace JobDriver.Execution
 
 		readonly PosixSignalRegistration? _sigTermHandler;
 
+		/// <summary>
+		/// Constructor
+		/// </summary>
 		protected JobExecutor(JobExecutorOptions options, Tracer tracer, ILogger logger)
 		{
 			HordeClient = options.HordeClient;
@@ -332,10 +386,13 @@ namespace JobDriver.Execution
 		/// <inheritdoc/>
 		public void Dispose()
 		{
-			GC.SuppressFinalize(this);
 			Dispose(true);
+			GC.SuppressFinalize(this);
 		}
 
+		/// <summary>
+		/// Disposes managed resources
+		/// </summary>
 		protected virtual void Dispose(bool disposing)
 		{
 			if (disposing)
@@ -364,6 +421,9 @@ namespace JobDriver.Execution
 			}
 		}
 
+		/// <summary>
+		/// Executes all steps in the job batch
+		/// </summary>
 		public async Task ExecuteAsync(ILogger logger, CancellationToken cancellationToken)
 		{
 			using TelemetrySpan span = Tracer.StartActiveSpan($"{nameof(JobExecutor)}.{nameof(ExecuteAsync)}");
@@ -405,6 +465,9 @@ namespace JobDriver.Execution
 			}, CancellationToken.None);
 		}
 
+		/// <summary>
+		/// Internal execution logic for the job batch, polling for and running steps
+		/// </summary>
 		protected virtual async Task ExecuteInternalAsync(ILogger logger, CancellationToken cancellationToken)
 		{
 			string? cleanupVar = Environment.GetEnvironmentVariable("UE_HORDE_JOB_BATCH_CLEANUP");
@@ -704,6 +767,9 @@ namespace JobDriver.Execution
 			}
 		}
 
+		/// <summary>
+		/// Initializes the executor with the batch response from the server
+		/// </summary>
 		public virtual async Task InitializeAsync(RpcBeginBatchResponse batch, ILogger logger, CancellationToken cancellationToken)
 		{
 			using TelemetrySpan span = Tracer.StartActiveSpan($"{nameof(JobExecutor)}.{nameof(InitializeAsync)}");
@@ -711,18 +777,18 @@ namespace JobDriver.Execution
 			Batch = batch;
 			JobRpc = await HordeClient.CreateGrpcClientAsync<JobRpc.JobRpcClient>(cancellationToken);
 
-			_envVars[HordeHttpClient.HordeUrlEnvVarName] = HordeClient.ServerUrl.ToString();
+			EnvVars[HordeHttpClient.HordeUrlEnvVarName] = HordeClient.ServerUrl.ToString();
 
 			string? accessToken = await HordeClient.GetAccessTokenAsync(false, cancellationToken);
 			if (accessToken != null)
 			{
-				_envVars[HordeHttpClient.HordeTokenEnvVarName] = accessToken;
+				EnvVars[HordeHttpClient.HordeTokenEnvVarName] = accessToken;
 			}
 
 			// Setup the agent type
 			foreach (KeyValuePair<string, string> envVar in Batch.Environment)
 			{
-				_envVars[envVar.Key] = envVar.Value;
+				EnvVars[envVar.Key] = envVar.Value;
 			}
 
 			// Figure out if we're running as an admin
@@ -749,38 +815,38 @@ namespace JobDriver.Execution
 						int colonIdx = map.IndexOf(':', StringComparison.Ordinal);
 						if (colonIdx != -1)
 						{
-							_remapAgentTypes[map.Substring(0, colonIdx)] = map.Substring(colonIdx + 1);
+							RemapAgentTypes[map.Substring(0, colonIdx)] = map.Substring(colonIdx + 1);
 						}
 					}
 				}
 				else if (argument.StartsWith(ScriptArgumentPrefix, StringComparison.OrdinalIgnoreCase))
 				{
-					_scriptFileName = argument.Substring(ScriptArgumentPrefix.Length);
+					ScriptFileName = argument.Substring(ScriptArgumentPrefix.Length);
 				}
 				else if (argument.Equals("-Preprocess", StringComparison.OrdinalIgnoreCase))
 				{
-					_preprocessScript = true;
+					PreprocessScript = true;
 				}
 				else if (argument.Equals("-SavePreprocessed", StringComparison.OrdinalIgnoreCase))
 				{
-					_savePreprocessedScript = true;
+					SavePreprocessedScript = true;
 				}
 				else if (argument.Equals("-AllowTargetChanges", StringComparison.OrdinalIgnoreCase))
 				{
-					_allowTargetChanges = true;
+					AllowTargetChanges = true;
 				}
 				else if (argument.StartsWith(TargetArgumentPrefix, StringComparison.OrdinalIgnoreCase))
 				{
-					_targets.Add(argument.Substring(TargetArgumentPrefix.Length));
+					Targets.Add(argument.Substring(TargetArgumentPrefix.Length));
 				}
 				else
 				{
-					_additionalArguments.Add(argument);
+					AdditionalArguments.Add(argument);
 				}
 			}
 			if (Batch.PreflightChange != 0)
 			{
-				_additionalArguments.Add($"-set:PreflightChange={Batch.PreflightChange}");
+				AdditionalArguments.Add($"-set:PreflightChange={Batch.PreflightChange}");
 			}
 		}
 
@@ -805,6 +871,9 @@ namespace JobDriver.Execution
 			}
 		}
 
+		/// <summary>
+		/// Runs a single job step
+		/// </summary>
 		public virtual async Task<JobStepOutcome> RunAsync(JobStepInfo step, ILogger logger, CancellationToken cancellationToken)
 		{
 			if (step.Name == "Setup Build")
@@ -831,6 +900,9 @@ namespace JobDriver.Execution
 			}
 		}
 
+		/// <summary>
+		/// Finalizes the executor after all steps have completed
+		/// </summary>
 		public virtual Task FinalizeAsync(ILogger jobLogger, CancellationToken cancellationToken)
 		{
 			using TelemetrySpan span = Tracer.StartActiveSpan($"{nameof(JobExecutor)}.{nameof(FinalizeAsync)}");
@@ -842,6 +914,9 @@ namespace JobDriver.Execution
 			return DirectoryReference.Combine(sharedStorageDir, "UAT");
 		}
 
+		/// <summary>
+		/// Copies pre-built AutomationTool binaries from shared storage into the workspace
+		/// </summary>
 		protected void CopyAutomationTool(DirectoryReference sharedStorageDir, DirectoryReference workspaceDir, ILogger logger)
 		{
 			DirectoryReference buildDir = GetAutomationToolDir(sharedStorageDir);
@@ -865,10 +940,13 @@ namespace JobDriver.Execution
 					DirectoryReference.CreateDirectory(targetFile.Directory);
 					FileReference.Copy(sourceFile, targetFile);
 				}
-				_compileAutomationTool = false;
+				CompileAutomationTool = false;
 			}
 		}
 
+		/// <summary>
+		/// Deletes cached BuildGraph manifests from the workspace
+		/// </summary>
 		protected void DeleteCachedBuildGraphManifests(DirectoryReference workspaceDir, ILogger logger)
 		{
 			DirectoryReference manifestDir = DirectoryReference.Combine(workspaceDir, $"{EnginePath}/Saved/BuildGraph");
@@ -885,8 +963,14 @@ namespace JobDriver.Execution
 			}
 		}
 
+		/// <summary>
+		/// Sets up the build environment for a step. Implemented by derived executors.
+		/// </summary>
 		protected abstract Task<bool> SetupAsync(JobStepInfo step, ILogger logger, CancellationToken cancellationToken);
 
+		/// <summary>
+		/// Executes a step. Implemented by derived executors.
+		/// </summary>
 		protected abstract Task<bool> ExecuteAsync(JobStepInfo step, ILogger logger, CancellationToken cancellationToken);
 
 		const string SetupStepName = "setup";
@@ -897,35 +981,38 @@ namespace JobDriver.Execution
 			return HordeClient.GetStorageNamespace(namespaceId, token);
 		}
 
+		/// <summary>
+		/// Sets up a build step by exporting the BuildGraph definition and uploading it to the server
+		/// </summary>
 		protected virtual async Task<bool> SetupAsync(JobStepInfo step, DirectoryReference workspaceDir, bool? useP4, ILogger logger, CancellationToken cancellationToken)
 		{
 			FileReference definitionFile = FileReference.Combine(workspaceDir, $"{EnginePath}/Saved/Horde/Exported.json");
 
 			StringBuilder arguments = new StringBuilder($"BuildGraph");
-			if (_scriptFileName != null)
+			if (ScriptFileName != null)
 			{
-				arguments.AppendArgument(ScriptArgumentPrefix, _scriptFileName);
+				arguments.AppendArgument(ScriptArgumentPrefix, ScriptFileName);
 			}
 			arguments.AppendArgument("-HordeExport=", definitionFile.FullName);
 			arguments.AppendArgument("-ListOnly");
 
-			if (!_allowTargetChanges)
+			if (!AllowTargetChanges)
 			{
-				foreach (string target in _targets)
+				foreach (string target in Targets)
 				{
 					arguments.AppendArgument(TargetArgumentPrefix + target);
 				}
 			}
 
 			//Arguments.AppendArgument("-TokenSignature=", JobId.ToString());
-			foreach (string additionalArgument in _additionalArguments)
+			foreach (string additionalArgument in AdditionalArguments)
 			{
 				arguments.AppendArgument(additionalArgument);
 			}
 
 			FileReference? preprocessedScriptFile = null;
 			FileReference? preprocessedSchemaFile = null;
-			if (_preprocessScript || _savePreprocessedScript)
+			if (PreprocessScript || SavePreprocessedScript)
 			{
 				preprocessedScriptFile = FileReference.Combine(workspaceDir, PreprocessedScript);
 				arguments.AppendArgument("-Preprocess=", preprocessedScriptFile.FullName);
@@ -1002,7 +1089,7 @@ namespace JobDriver.Execution
 			validTargets.UnionWith(updateGraph.Groups.SelectMany(x => x.Nodes).Select(x => x.Name));
 			validTargets.UnionWith(updateGraph.Aggregates.Select(x => x.Name));
 
-			foreach (string target in _targets)
+			foreach (string target in Targets)
 			{
 				if (!validTargets.Contains(target))
 				{
@@ -1031,7 +1118,7 @@ namespace JobDriver.Execution
 				foreach (string validAgentTypeName in exportedGroup.Types)
 				{
 					string? thisAgentTypeName;
-					if (!_remapAgentTypes.TryGetValue(validAgentTypeName, out thisAgentTypeName))
+					if (!RemapAgentTypes.TryGetValue(validAgentTypeName, out thisAgentTypeName))
 					{
 						thisAgentTypeName = validAgentTypeName;
 					}
@@ -1205,10 +1292,13 @@ namespace JobDriver.Execution
 			}
 		}
 
+		/// <summary>
+		/// Executes a build step by running BuildGraph with the given workspace
+		/// </summary>
 		protected async Task<bool> ExecuteAsync(JobStepInfo step, DirectoryReference workspaceDir, bool? useP4, ILogger logger, CancellationToken cancellationToken)
 		{
 			StringBuilder arguments = new StringBuilder("BuildGraph");
-			if (_preprocessScript)
+			if (PreprocessScript)
 			{
 				FileReference localPreprocessedScript = FileReference.Combine(workspaceDir, PreprocessedScript);
 				arguments.AppendArgument(ScriptArgumentPrefix, localPreprocessedScript.FullName);
@@ -1245,14 +1335,14 @@ namespace JobDriver.Execution
 					logger.LogInformation("Bundle has no {BuildGraphFolderName} folder; not copying any files", BuildGraphTempStorageDir);
 				}
 			}
-			else if (_scriptFileName != null)
+			else if (ScriptFileName != null)
 			{
-				arguments.AppendArgument(ScriptArgumentPrefix, _scriptFileName);
+				arguments.AppendArgument(ScriptArgumentPrefix, ScriptFileName);
 			}
 
-			if (!_allowTargetChanges)
+			if (!AllowTargetChanges)
 			{
-				foreach (string target in _targets)
+				foreach (string target in Targets)
 				{
 					arguments.AppendArgument(TargetArgumentPrefix + target);
 				}
@@ -1261,9 +1351,9 @@ namespace JobDriver.Execution
 			arguments.AppendArgument("-SingleNode=", step.Name);
 			//			Arguments.AppendArgument("-TokenSignature=", JobId.ToString());
 
-			foreach (string additionalArgument in _additionalArguments)
+			foreach (string additionalArgument in AdditionalArguments)
 			{
-				if (!_preprocessScript || !additionalArgument.StartsWith("-set:", StringComparison.OrdinalIgnoreCase))
+				if (!PreprocessScript || !additionalArgument.StartsWith("-set:", StringComparison.OrdinalIgnoreCase))
 				{
 					arguments.AppendArgument(additionalArgument);
 				}
@@ -1273,6 +1363,9 @@ namespace JobDriver.Execution
 			return result;
 		}
 		
+		/// <summary>
+		/// Creates and uploads an artifact for a job step
+		/// </summary>
 		protected async Task CreateArtifactAsync(JobStepId stepId, ArtifactName name, ArtifactType type, DirectoryReference baseDir, IReadOnlyList<FileInfo> files, ILogger logger, CancellationToken cancellationToken)
 		{
 			try
@@ -1630,6 +1723,9 @@ namespace JobDriver.Execution
 			return true;
 		}
 
+		/// <summary>
+		/// Executes the Unreal AutomationTool (UAT) with the given arguments
+		/// </summary>
 		protected async Task<int> ExecuteAutomationToolAsync(JobStepInfo step, DirectoryReference workspaceDir, string? arguments, bool? useP4, ILogger logger, CancellationToken cancellationToken)
 		{
 			int result;
@@ -1637,7 +1733,7 @@ namespace JobDriver.Execution
 			span.SetAttribute("horde.job.step.id", step.StepId.ToString());
 			span.SetAttribute("horde.job.step.name", step.Name);
 
-			if (!_compileAutomationTool)
+			if (!CompileAutomationTool)
 			{
 				arguments += " -NoCompile";
 			}
@@ -1673,7 +1769,7 @@ namespace JobDriver.Execution
 				throw new Exception("Unsupported platform");
 			}
 
-			_compileAutomationTool = false;
+			CompileAutomationTool = false;
 			return result;
 		}
 
@@ -1747,9 +1843,12 @@ namespace JobDriver.Execution
 
 		internal IReadOnlyDictionary<string, string> GetEnvVars()
 		{
-			return new Dictionary<string, string>(_envVars);
+			return new Dictionary<string, string>(EnvVars);
 		}
 
+		/// <summary>
+		/// Runs the lease cleanup script in the workspace if present
+		/// </summary>
 		protected static async Task ExecuteLeaseCleanupScriptAsync(DirectoryReference workspaceDir, ILogger logger)
 		{
 			using (LogParser parser = new LogParser(logger, new List<string>()))
@@ -1759,6 +1858,9 @@ namespace JobDriver.Execution
 			}
 		}
 
+		/// <summary>
+		/// Terminates processes matching the configured termination rules
+		/// </summary>
 		protected async Task TerminateProcessesAsync(TerminateCondition condition, ILogger logger)
 		{
 			try
@@ -1951,7 +2053,7 @@ namespace JobDriver.Execution
 
 			// Combine all the supplied environment variables together
 			Dictionary<string, string> newEnvVars = new Dictionary<string, string>(StringComparer.Ordinal);
-			foreach (KeyValuePair<string, string> envVar in _envVars)
+			foreach (KeyValuePair<string, string> envVar in EnvVars)
 			{
 				newEnvVars[ConformEnvironmentVariableName(envVar.Key)] = StringUtils.ExpandProperties(envVar.Value, properties);
 			}
@@ -2389,6 +2491,9 @@ namespace JobDriver.Execution
 			span.SetAttribute(prefix + "artifact.namespace_id", response.NamespaceId);
 		}
 
+		/// <summary>
+		/// Uploads test data for a job step to the server
+		/// </summary>
 		public async Task UploadTestDataAsync(JobStepId jobStepId, IEnumerable<KeyValuePair<string, object>> pairs)
 		{
 			if (pairs.Any())
@@ -2415,6 +2520,9 @@ namespace JobDriver.Execution
 			}
 		}
 
+		/// <summary>
+		/// Gets the Unreal root directory from a workspace directory
+		/// </summary>
 		protected DirectoryReference GetUnrealRootDirectoryFromWorkspace(DirectoryReference workspaceDir)
 		{
 			// The Unreal root directory is the parent of the Engine directory
@@ -2422,10 +2530,19 @@ namespace JobDriver.Execution
 		}
 	}
 
-	interface IJobExecutorFactory
+	/// <summary>
+	/// Factory interface for creating job executors
+	/// </summary>
+	public interface IJobExecutorFactory
 	{
+		/// <summary>
+		/// Name of the executor type
+		/// </summary>
 		string Name { get; }
 
+		/// <summary>
+		/// Creates a new executor instance for the given workspace configuration
+		/// </summary>
 		Task<JobExecutor> CreateExecutorAsync(RpcAgentWorkspace workspaceInfo, RpcAgentWorkspace? autoSdkWorkspaceInfo, JobExecutorOptions options, CancellationToken cancellationToken);
 	}
 
diff --git a/Engine/Source/Programs/Horde/Drivers/JobDriver/Execution/PerforceExecutor.cs b/Engine/Source/Programs/Horde/Drivers/JobDriver/Execution/PerforceExecutor.cs
--- a/Engine/Source/Programs/Horde/Drivers/JobDriver/Execution/PerforceExecutor.cs
+++ b/Engine/Source/Programs/Horde/Drivers/JobDriver/Execution/PerforceExecutor.cs
@@ -150,22 +150,22 @@ namespace JobDriver.Execution
 			DeleteEngineUserSettings(logger);
 
 			// Get all the environment variables for jobs
-			_envVars["IsBuildMachine"] = "1";
-			_envVars["uebp_LOCAL_ROOT"] = GetUnrealRootDirectoryFromWorkspace(_workspace.WorkspaceDir).FullName;
-			_envVars["uebp_PORT"] = _workspace.ServerAndPort;
-			_envVars["uebp_USER"] = _workspace.UserName;
-			_envVars["uebp_CLIENT"] = _workspace.ClientName;
-			_envVars["uebp_BuildRoot_P4"] = Batch.StreamName;
-			_envVars["uebp_BuildRoot_Escaped"] = Batch.StreamName.Replace('/', '+');
-			_envVars["uebp_CLIENT_ROOT"] = $"//{_workspace.ClientName}";
-			_envVars["uebp_CL"] = Batch.Change.ToString();
-			_envVars["uebp_CodeCL"] = Batch.CodeChange.ToString();
-			_envVars["P4USER"] = _workspace.UserName;
-			_envVars["P4CLIENT"] = _workspace.ClientName;
+			EnvVars["IsBuildMachine"] = "1";
+			EnvVars["uebp_LOCAL_ROOT"] = GetUnrealRootDirectoryFromWorkspace(_workspace.WorkspaceDir).FullName;
+			EnvVars["uebp_PORT"] = _workspace.ServerAndPort;
+			EnvVars["uebp_USER"] = _workspace.UserName;
+			EnvVars["uebp_CLIENT"] = _workspace.ClientName;
+			EnvVars["uebp_BuildRoot_P4"] = Batch.StreamName;
+			EnvVars["uebp_BuildRoot_Escaped"] = Batch.StreamName.Replace('/', '+');
+			EnvVars["uebp_CLIENT_ROOT"] = $"//{_workspace.ClientName}";
+			EnvVars["uebp_CL"] = Batch.Change.ToString();
+			EnvVars["uebp_CodeCL"] = Batch.CodeChange.ToString();
+			EnvVars["P4USER"] = _workspace.UserName;
+			EnvVars["P4CLIENT"] = _workspace.ClientName;
 
 			if (_autoSdkWorkspace != null)
 			{
-				_envVars["UE_SDKS_ROOT"] = _autoSdkWorkspace.WorkspaceDir.FullName;
+				EnvVars["UE_SDKS_ROOT"] = _autoSdkWorkspace.WorkspaceDir.FullName;
 			}
 		}
 
diff --git a/Engine/Source/Programs/Horde/Drivers/JobDriver/Execution/WorkspaceExecutor.cs b/Engine/Source/Programs/Horde/Drivers/JobDriver/Execution/WorkspaceExecutor.cs
--- a/Engine/Source/Programs/Horde/Drivers/JobDriver/Execution/WorkspaceExecutor.cs
+++ b/Engine/Source/Programs/Horde/Drivers/JobDriver/Execution/WorkspaceExecutor.cs
@@ -76,21 +76,21 @@ namespace JobDriver.Execution
 			}
 
 			// Set any non-materializer specific environment variables for jobs
-			_envVars["IsBuildMachine"] = "1";
-			_envVars["uebp_LOCAL_ROOT"] = GetUnrealRootDirectoryFromWorkspace(_workspace.SyncDir).FullName;
-			_envVars["uebp_BuildRoot_P4"] = Batch.StreamName;
-			_envVars["uebp_BuildRoot_Escaped"] = Batch.StreamName.Replace('/', '+');
-			_envVars["uebp_CL"] = Batch.Change.ToString();
-			_envVars["uebp_CodeCL"] = Batch.CodeChange.ToString();
+			EnvVars["IsBuildMachine"] = "1";
+			EnvVars["uebp_LOCAL_ROOT"] = GetUnrealRootDirectoryFromWorkspace(_workspace.SyncDir).FullName;
+			EnvVars["uebp_BuildRoot_P4"] = Batch.StreamName;
+			EnvVars["uebp_BuildRoot_Escaped"] = Batch.StreamName.Replace('/', '+');
+			EnvVars["uebp_CL"] = Batch.Change.ToString();
+			EnvVars["uebp_CodeCL"] = Batch.CodeChange.ToString();
 
 			foreach ((string key, string value) in _workspace.EnvironmentVariables)
 			{
-				_envVars[key] = value;
+				EnvVars[key] = value;
 			}
 
 			if (_autoSdkWorkspace != null)
 			{
-				_envVars["UE_SDKS_ROOT"] = _autoSdkWorkspace.SyncDir.FullName;
+				EnvVars["UE_SDKS_ROOT"] = _autoSdkWorkspace.SyncDir.FullName;
 			}
 		}
 
diff --git a/Engine/Source/Programs/Horde/HordeDashboard/src/backend/Api.ts b/Engine/Source/Programs/Horde/HordeDashboard/src/backend/Api.ts
--- a/Engine/Source/Programs/Horde/HordeDashboard/src/backend/Api.ts
+++ b/Engine/Source/Programs/Horde/HordeDashboard/src/backend/Api.ts
@@ -1141,12 +1141,18 @@ export type CreateJobRequest = {
 	/** The changelist number to build. Can be null for latest. */
 	change?: number;
 
+	/** The commit id to build. Can be null for latest. */
+	commitId?: string;
+
 	/// List of change queries to evaluate
 	changeQueries?: ChangeQueryConfig[];
 
 	/** The preflight changelist number */
 	preflightChange?: number;
 
+	/** Commit hash for branch-based preflights */
+	preflightCommitId?: string;
+
 	/** Priority for the job */
 	priority?: Priority;
 
@@ -1277,6 +1283,9 @@ export type GetJobResponse = {
 	/** Description of the preflight */
 	preflightDescription?: string;
 
+	/** Commit hash for branch-based preflights */
+	preflightCommitId?: string;
+
 	/** The template type */
 	templateId?: string;
 
diff --git a/Engine/Source/Programs/Horde/HordeDashboard/src/components/JobViewIncremental.tsx b/Engine/Source/Programs/Horde/HordeDashboard/src/components/JobViewIncremental.tsx
--- a/Engine/Source/Programs/Horde/HordeDashboard/src/components/JobViewIncremental.tsx
+++ b/Engine/Source/Programs/Horde/HordeDashboard/src/components/JobViewIncremental.tsx
@@ -516,6 +516,8 @@ const JobList: React.FC<{ tab: string; filter: JobFilterSimple; aliasTemplates?:
          let change = item.job?.change?.toString() ?? "Latest";
          if (item.job?.preflightChange) {
             change = `PF ${item.job.preflightChange}`;
+         } else if (item.job?.preflightCommitId) {
+            change = `PF ${item.job.preflightCommitId.substring(0, 8)}`;
          }
 
 
diff --git a/Engine/Source/Programs/Horde/HordeDashboard/src/components/Preflight.tsx b/Engine/Source/Programs/Horde/HordeDashboard/src/components/Preflight.tsx
--- a/Engine/Source/Programs/Horde/HordeDashboard/src/components/Preflight.tsx
+++ b/Engine/Source/Programs/Horde/HordeDashboard/src/components/Preflight.tsx
@@ -55,8 +55,9 @@ export const PreflightRedirector: React.FC = () => {
    }
 
    const cl = parseInt(change);
+   const isNumericChange = !isNaN(cl);
 
-   if (isNaN(cl)) {
+   if (!isNumericChange && !/^[a-f0-9]+$/i.test(change)) {
       setError(`Bad change in preflight ${change}`);
       return null;
    }
@@ -114,40 +115,8 @@ export const PreflightRedirector: React.FC = () => {
 
       console.log(`Redirecting preflight: ${window.location.href}`);
 
-      backend.getJobs({ filter: "id,streamId", count: 1, preflightChange: cl }).then(result => {
-
-         if (result && result.length === 1) {
-
-            if (stream.id === result[0].streamId) {
-
-               let url = `/job/${result[0].id}?newbuild=true&allowtemplatechange=true&shelvedchange=${change}&p4v=true`;
-
-               if (autosubmit === "true") {
-                  url += "&autosubmit=true";
-               }
-
-               if (templateId) {
-                  url += `&templateId=${templateId}`;
-               }
-
-               if (parameters.length) {
-                  url += ("&" + parameters.join("&"));
-               }
-
-               if (version === "2") {
-                  url += "&newbuildversion=2"
-               }
-
-               if (resolvedAlias) {
-                  url += `&alias=${encodeURIComponent(resolvedAlias)}`;
-               }
-
-               navigate(url, { replace: true });
-               return;
-            } 
-         }
-
-         let url = `/stream/${stream!.id}?tab=${tab}&newbuild=true&shelvedchange=${change}&p4v=true`;
+      const buildRedirectUrl = (basePath: string) => {
+         let url = `${basePath}?newbuild=true&shelvedchange=${encodeURIComponent(change)}&p4v=true`;
 
          if (autosubmit === "true") {
             url += "&autosubmit=true";
@@ -169,8 +138,25 @@ export const PreflightRedirector: React.FC = () => {
             url += `&alias=${encodeURIComponent(resolvedAlias)}`;
          }
 
-         navigate(url, { replace: true });
+         return url;
+      };
+
+      // For numeric CLs, check if there's an existing job to redirect to
+      const jobLookup = isNumericChange
+         ? backend.getJobs({ filter: "id,streamId", count: 1, preflightChange: cl })
+         : Promise.resolve([]);
+
+      jobLookup.then(result => {
+
+         if (result && result.length === 1) {
+
+            if (stream.id === result[0].streamId) {
+               navigate(buildRedirectUrl(`/job/${result[0].id}`) + "&allowtemplatechange=true", { replace: true });
+               return;
+            }
+         }
 
+         navigate(buildRedirectUrl(`/stream/${stream!.id}`) + `&tab=${tab}`, { replace: true });
 
       }).catch(reason => {
          console.error(`Error getting job for preflight: `, reason);
diff --git a/Engine/Source/Programs/Horde/HordeDashboard/src/components/build/NewBuildV2.tsx b/Engine/Source/Programs/Horde/HordeDashboard/src/components/build/NewBuildV2.tsx
--- a/Engine/Source/Programs/Horde/HordeDashboard/src/components/build/NewBuildV2.tsx
+++ b/Engine/Source/Programs/Horde/HordeDashboard/src/components/build/NewBuildV2.tsx
@@ -103,8 +103,11 @@ class BuildOptions {
 
       this.change = undefined;
       this.changeOption = undefined;
+      this.changeCommitId = undefined;
       this.preflightChange = undefined;
+      this.preflightCommitId = undefined;
       this.queryShelvedChange = undefined;
+      this.queryShelvedCommitId = undefined;
       this.autoSubmit = undefined;
       this.advJobName = undefined;
       this.updateIssues = undefined;
@@ -127,15 +130,23 @@ class BuildOptions {
       this.isPreflightSubmit = !!query.get("shelvedchange");
       this.isFromP4V = !!query.get("p4v");
       this.queryTemplateId = !query.get("templateId") ? "" : query.get("templateId")!;
-      const queryChange = parseInt(query.get("shelvedchange") ? query.get("shelvedchange")! : "0");
-      if (!isNaN(queryChange) && queryChange) {
-         this.queryShelvedChange = queryChange;
+      const queryShelvedParam = query.get("shelvedchange") ?? "";
+      if (queryShelvedParam) {
+         const parsed = parseInt(queryShelvedParam);
+         if (!isNaN(parsed) && parsed) {
+            this.queryShelvedChange = parsed;
+         } else if (/^[a-f0-9]+$/i.test(queryShelvedParam)) {
+            this.queryShelvedCommitId = queryShelvedParam.trim().toLowerCase();
+         }
       }
 
       this.preflightChange = jobDetails?.jobData?.preflightChange;
+      this.preflightCommitId = jobDetails?.jobData?.preflightCommitId;
       this.change = jobDetails?.jobData?.change;
-      if (!this.preflightChange) {
+
+      if (!this.preflightChange && !this.preflightCommitId) {
          this.preflightChange = this.queryShelvedChange;
+         this.preflightCommitId = this.queryShelvedCommitId;
       }
 
       this.advJobPriority = jobDetails?.jobData?.priority;
@@ -249,8 +260,11 @@ class BuildOptions {
 
    change?: number;
    changeOption?: string;
+   changeCommitId?: string;
    preflightChange?: number;
+   preflightCommitId?: string;
    queryShelvedChange?: number;
+   queryShelvedCommitId?: string;
    autoSubmit?: boolean;
 
    static initialized() {
@@ -343,6 +357,21 @@ class BuildOptions {
       }
 
       this.preflightChange = preflightCL;
+      if (preflightCL) {
+         this.preflightCommitId = undefined;
+      }
+
+      this.setChanged();
+   }
+
+   onPreflightCommitIdChange(commitId?: string) {
+      if (this.preflightCommitId === commitId) {
+         return;
+      }
+      this.preflightCommitId = commitId;
+      if (commitId) {
+         this.preflightChange = undefined;
+      }
 
       this.setChanged();
    }
@@ -479,7 +508,7 @@ class BuildOptions {
       }
 
       if (this.preflightChange && !this.template.allowPreflights) {
-         errors.push({ error: `Template "${this.template.name}" does not allow preflights`, paramString: "Shelved Change", value: this.preflightChange.toString() })
+         errors.push({ error: `Template "${this.template.name}" does not allow preflights`, paramString: "Preflight Change", value: this.preflightChange.toString() })
       }
 
       if (this.change !== undefined && isNaN(this.change)) {
@@ -492,12 +521,16 @@ class BuildOptions {
 
       if (this.preflightChange !== undefined && isNaN(this.preflightChange)) {
          errors.push({
-            paramString: "Shelved Change",
+            paramString: "Preflight Change",
             value: this.preflightChange.toString(),
             error: "Invalid shelved changelist"
          });
       }
 
+      if (this.preflightCommitId && !this.template.allowPreflights) {
+         errors.push({ error: `Template "${this.template.name}" does not allow preflights`, paramString: "Preflight Change", value: this.preflightCommitId })
+      }
+
       this.template.parameters.forEach(p => {
 
          this.validateParameter(p, errors);
@@ -917,7 +950,7 @@ class BuildOptions {
          t = this.allTemplates.find(t => t.name === jobDetails.template!.name);
       }
       // handle preflight redirect case
-      if (!t && this.queryShelvedChange && !jobDetails) {
+      if (!t && this.isPreflightSubmit && !jobDetails) {
 
          if (queryTemplateId) {
             let errorReason = "";
@@ -1992,6 +2025,9 @@ const BuildModal: React.FC<{ setUseLegacyDialog: (value: boolean) => void }> = o
    if (options.change) {
       changeText = options.change.toString();
    }
+   if (options.changeCommitId) {
+      changeText = options.changeCommitId;
+   }
    if (options.changeOption) {
       changeText = options.changeOption;
    }
@@ -2006,7 +2042,7 @@ const BuildModal: React.FC<{ setUseLegacyDialog: (value: boolean) => void }> = o
 
       let changeQueries: ChangeQueryConfig[] | undefined;
 
-      if (typeof options.change !== 'number') {
+      if (typeof options.change !== 'number' && !options.changeCommitId) {
          const changeOption = options.changeOption ?? "Latest Change";
          if (changeOption !== "Latest Change" && defaultPreflightQuery) {
             changeQueries = defaultPreflightQuery;
@@ -2056,6 +2092,10 @@ const BuildModal: React.FC<{ setUseLegacyDialog: (value: boolean) => void }> = o
          data.change = options.change;
       }
 
+      if (options.changeCommitId) {
+         data.commitId = options.changeCommitId;
+      }
+
       if (typeof (options.preflightChange) === 'number') {
 
          data.preflightChange = options.preflightChange;
@@ -2068,6 +2108,11 @@ const BuildModal: React.FC<{ setUseLegacyDialog: (value: boolean) => void }> = o
 
       }
 
+      if (options.preflightCommitId) {
+         data.preflightCommitId = options.preflightCommitId;
+         data.updateIssues = false;
+      }
+
       const submit = async () => {
 
          if (!!import.meta.env.VITE_HORDE_DEBUG_NEW_JOB) {
@@ -2190,14 +2235,22 @@ const BuildModal: React.FC<{ setUseLegacyDialog: (value: boolean) => void }> = o
                            ev.preventDefault();
                            if (option) {
                               options.change = undefined;
+                              options.changeCommitId = undefined;
                               options.changeOption = option.text;
                            } else {
                               options.changeOption = undefined;
                               if (!value) {
                                  options.change = undefined;
+                                 options.changeCommitId = undefined;
                               } else {
                                  const nvalue = parseInt(value);
-                                 options.change = !isNaN(nvalue) ? nvalue : undefined;
+                                 if (!isNaN(nvalue)) {
+                                    options.change = nvalue;
+                                    options.changeCommitId = undefined;
+                                 } else {
+                                    options.change = undefined;
+                                    options.changeCommitId = value.trim().toLowerCase();
+                                 }
                               }
                            }
                            options.setChanged();
@@ -2208,24 +2261,29 @@ const BuildModal: React.FC<{ setUseLegacyDialog: (value: boolean) => void }> = o
 
                      <Stack>
                         <TextField
-                           key={`build_shelved_change_key_${options.currentRenderKey}`}
+                           key={`build_preflight_key_${options.currentRenderKey}`}
                            label="Shelved Change"
-                           style={{ width: 146 }}
-                           placeholder={!template!.allowPreflights ? "Disabled by template" : "None"}
-                           title={!template!.allowPreflights ? "Preflights are disabled for this template" : undefined}
-                           autoComplete="off" value={options.preflightChange?.toString() ?? ""} disabled={!template!.allowPreflights || options.readOnly || options.isPreflightSubmit} onChange={(ev, newValue) => {
+                           style={{ width: 120 }}
+                           placeholder={!template!.allowPreflights ? "Disabled" : ""}
+                           autoComplete="off"
+                           spellCheck={false}
+                           value={options.preflightChange?.toString() ?? options.preflightCommitId ?? ""}
+                           disabled={!template!.allowPreflights || options.readOnly || options.isPreflightSubmit}
+                           onChange={(ev, newValue) => {
                               ev.preventDefault();
-
-                              let change: number | undefined;
-                              if (newValue) {
-                                 change = parseInt(newValue);
-                                 if (isNaN(change)) {
-                                    change = undefined;
-                                 }
+                              const trimmed = newValue?.trim() ?? "";
+                              if (!trimmed) {
+                                 options.onPreflightChange(undefined);
+                                 options.onPreflightCommitIdChange(undefined);
+                                 return;
                               }
-                              options.onPreflightChange(change);
-                           }} />
-
+                              if (/^\d+$/.test(trimmed)) {
+                                 options.onPreflightChange(parseInt(trimmed));
+                              } else {
+                                 options.onPreflightCommitIdChange(trimmed.toLowerCase());
+                              }
+                           }}
+                        />
                      </Stack>
                   </Stack>
                </Stack>
diff --git a/Engine/Source/Programs/Horde/HordeDashboard/src/components/jobDetailsV2/JobDetailPreflight.tsx b/Engine/Source/Programs/Horde/HordeDashboard/src/components/jobDetailsV2/JobDetailPreflight.tsx
--- a/Engine/Source/Programs/Horde/HordeDashboard/src/components/jobDetailsV2/JobDetailPreflight.tsx
+++ b/Engine/Source/Programs/Horde/HordeDashboard/src/components/jobDetailsV2/JobDetailPreflight.tsx
@@ -1,6 +1,6 @@
 // Copyright Epic Games, Inc. All Rights Reserved.
 
-import { Checkbox, Stack, Text } from '@fluentui/react';
+import { Checkbox, IconButton, Stack, Text } from '@fluentui/react';
 import { observer } from 'mobx-react-lite';
 import React, { useEffect, useState } from 'react';
 import backend from '../../backend';
@@ -118,6 +118,8 @@ export const PreflightPanel: React.FC<{ jobDetails: JobDetailsV2 }> = observer((
 
    const preflightView = jobDetails.getDataView<PreflightDataView>("PreflightDataView");
 
+   const [commitIdCopied, setCommitIdCopied] = React.useState(false);
+
    useEffect(() => {
       return () => {
          preflightView.clear();
@@ -130,7 +132,7 @@ export const PreflightPanel: React.FC<{ jobDetails: JobDetailsV2 }> = observer((
       return null;
    }
 
-   const hasPreflight = !!jobData.preflightChange || !!jobData.preflightDescription;
+   const hasPreflight = !!jobData.preflightChange || !!jobData.preflightDescription || !!jobData.preflightCommitId;
 
    if (!preflightView.initialized) {
       preflightView.initialize(hasPreflight ? [sideRail] : undefined);
@@ -145,7 +147,33 @@ export const PreflightPanel: React.FC<{ jobDetails: JobDetailsV2 }> = observer((
          <Stack tokens={{ childrenGap: 12 }} grow>
             <Stack horizontal>
                <Stack>
-                  <Text variant="mediumPlus" styles={{ root: { fontFamily: "Horde Open Sans SemiBold" } }}>{`Preflight (CL ${jobData.preflightChange})`}</Text>
+                  {jobData.preflightChange && (
+                     <Text variant="mediumPlus" styles={{ root: { fontFamily: "Horde Open Sans SemiBold" } }}>{`Preflight (CL ${jobData.preflightChange})`}</Text>
+                  )}
+                  {jobData.preflightCommitId && !jobData.preflightChange && (
+                     <Stack horizontal tokens={{ childrenGap: 8 }} verticalAlign="center">
+                        <Text variant="mediumPlus" styles={{ root: { fontFamily: "Horde Open Sans SemiBold" } }}>
+                           Preflight Commit:
+                        </Text>
+                        <Text variant="medium" styles={{ root: { fontFamily: 'monospace' } }}>
+                           {jobData.preflightCommitId.substring(0, 12)}...
+                        </Text>
+                        <IconButton
+                           iconProps={{ iconName: commitIdCopied ? 'CheckMark' : 'Copy' }}
+                           title={commitIdCopied ? "Copied!" : "Copy full commit ID"}
+                           styles={commitIdCopied ? { icon: { color: 'green' } } : undefined}
+                           onClick={async () => {
+                              try {
+                                 await navigator.clipboard.writeText(jobData.preflightCommitId!);
+                                 setCommitIdCopied(true);
+                                 setTimeout(() => setCommitIdCopied(false), 2000);
+                              } catch (err) {
+                                 console.error('Failed to copy commit ID to clipboard:', err);
+                              }
+                           }}
+                        />
+                     </Stack>
+                  )}
                </Stack>
             </Stack>
             <Stack >
diff --git a/Engine/Source/Programs/Horde/Plugins/Build/HordeServer.Build.Tests/Stubs/Services/PerforceServiceStub.cs b/Engine/Source/Programs/Horde/Plugins/Build/HordeServer.Build.Tests/Stubs/Services/PerforceServiceStub.cs
--- a/Engine/Source/Programs/Horde/Plugins/Build/HordeServer.Build.Tests/Stubs/Services/PerforceServiceStub.cs
+++ b/Engine/Source/Programs/Horde/Plugins/Build/HordeServer.Build.Tests/Stubs/Services/PerforceServiceStub.cs
@@ -13,8 +13,10 @@ using EpicGames.Perforce;
 
 namespace HordeServer.Tests.Stubs.Services
 {
-	class PerforceServiceStub : IPerforceService
+	public class PerforceServiceStub : IPerforceService
 	{
+		public string Name { get; } = "Perforce";
+		
 		class User : IUser
 		{
 			public UserId Id { get; }
diff --git a/Engine/Source/Programs/Horde/Plugins/Build/HordeServer.Build/BuildPlugin.cs b/Engine/Source/Programs/Horde/Plugins/Build/HordeServer.Build/BuildPlugin.cs
--- a/Engine/Source/Programs/Horde/Plugins/Build/HordeServer.Build/BuildPlugin.cs
+++ b/Engine/Source/Programs/Horde/Plugins/Build/HordeServer.Build/BuildPlugin.cs
@@ -41,6 +41,7 @@ using HordeServer.Users;
 using Microsoft.AspNetCore.Builder;
 using Microsoft.Extensions.DependencyInjection;
 using MongoDB.Bson.Serialization;
+using HordeServer.VersionControl;
 
 namespace HordeServer
 {
@@ -178,6 +179,9 @@ namespace HordeServer
 				services.AddSingleton<PerforceService>();
 				services.AddSingleton<IPerforceService>(sp => sp.GetRequiredService<PerforceService>());
 			}
+
+			services.AddSingleton<IVersionControlService>(sp => sp.GetRequiredService<IPerforceService>());
+
 			services.AddSingleton<PerforceReplicator>();
 
 			services.AddSingleton<PerforceLoadBalancer>();
diff --git a/Engine/Source/Programs/Horde/Plugins/Build/HordeServer.Build/Commits/CommitService.cs b/Engine/Source/Programs/Horde/Plugins/Build/HordeServer.Build/Commits/CommitService.cs
--- a/Engine/Source/Programs/Horde/Plugins/Build/HordeServer.Build/Commits/CommitService.cs
+++ b/Engine/Source/Programs/Horde/Plugins/Build/HordeServer.Build/Commits/CommitService.cs
@@ -2,31 +2,47 @@
 
 using EpicGames.Horde.Commits;
 using EpicGames.Horde.Streams;
-using HordeServer.VersionControl.Perforce;
 using HordeServer.Streams;
 using Microsoft.Extensions.Options;
+using HordeServer.VersionControl;
 
 namespace HordeServer.Commits
 {
 	/// <summary>
 	/// Provides commit information for streams
 	/// </summary>
-	class CommitService : ICommitService
+	public class CommitService : ICommitService
 	{
-		readonly IPerforceService _perforceService;
 		readonly IOptionsMonitor<BuildConfig> _buildConfig;
 
+		readonly Dictionary<string, IVersionControlService> _vcsServices = new Dictionary<string, IVersionControlService>(StringComparer.OrdinalIgnoreCase);
+
 		/// <summary>
 		/// Constructor
 		/// </summary>
-		public CommitService(IPerforceService perforceService, IOptionsMonitor<BuildConfig> buildConfig)
+		public CommitService(IEnumerable<IVersionControlService> versionControlServices, IOptionsMonitor<BuildConfig> buildConfig)
 		{
-			_perforceService = perforceService;
 			_buildConfig = buildConfig;
+
+			foreach (IVersionControlService service in versionControlServices)
+			{
+				if (!_vcsServices.TryAdd(service.Name, service))
+				{
+					throw new InvalidOperationException($"Duplicate VCS service registered with name '{service.Name}' (existing: {_vcsServices[service.Name].GetType().Name}, new: {service.GetType().Name})");
+				}
+			}
 		}
 
 		/// <inheritdoc/>
-		public ICommitCollection GetCollection(StreamConfig streamConfig) => _perforceService.GetCommits(streamConfig);
+		public ICommitCollection GetCollection(StreamConfig streamConfig)
+		{
+			if (!_vcsServices.TryGetValue(streamConfig.VCS, out IVersionControlService? vcs))
+			{
+				throw new KeyNotFoundException($"Unknown VCS Service {streamConfig.VCS}");
+			}
+
+			return vcs.GetCommits(streamConfig);
+		}
 
 		/// <inheritdoc/>
 		public async ValueTask<CommitIdWithOrder> GetOrderedAsync(StreamId streamId, CommitId commitId, CancellationToken cancellationToken = default)
diff --git a/Engine/Source/Programs/Horde/Plugins/Build/HordeServer.Build/Jobs/IWorkspaceMessageEnricher.cs b/Engine/Source/Programs/Horde/Plugins/Build/HordeServer.Build/Jobs/IWorkspaceMessageEnricher.cs
new file mode 100644
--- /dev/null
+++ b/Engine/Source/Programs/Horde/Plugins/Build/HordeServer.Build/Jobs/IWorkspaceMessageEnricher.cs
@@ -0,0 +1,27 @@
+// Copyright Epic Games, Inc. All Rights Reserved.
+
+using EpicGames.Horde.Jobs;
+using HordeCommon.Rpc.Messages;
+using HordeServer.Agents;
+
+namespace HordeServer.Jobs
+{
+	/// <summary>
+	/// Interface for enriching workspace messages before they are sent to agents.
+	/// Plugins can implement this to modify workspace credentials, add custom fields, etc.
+	/// </summary>
+	public interface IWorkspaceMessageEnricher
+	{
+		/// <summary>
+		/// Enriches a workspace message before it's sent to an agent.
+		/// Called after the base workspace message is created but before it's assigned to a task.
+		/// </summary>
+		/// <param name="workspace">The workspace message to enrich</param>
+		/// <param name="workspaceInfo">The workspace configuration info</param>
+		/// <param name="agent">The agent that will receive this workspace</param>
+		/// <param name="job">The job being executed (provides access to commit info for VCS-specific enrichment)</param>
+		/// <param name="cancellationToken">Cancellation token</param>
+		/// <returns>Task that completes when enrichment is done</returns>
+		Task EnrichAsync(RpcAgentWorkspace workspace, AgentWorkspaceInfo workspaceInfo, IAgent agent, IJob job, CancellationToken cancellationToken);
+	}
+}
diff --git a/Engine/Source/Programs/Horde/Plugins/Build/HordeServer.Build/Jobs/JobRpcService.cs b/Engine/Source/Programs/Horde/Plugins/Build/HordeServer.Build/Jobs/JobRpcService.cs
--- a/Engine/Source/Programs/Horde/Plugins/Build/HordeServer.Build/Jobs/JobRpcService.cs
+++ b/Engine/Source/Programs/Horde/Plugins/Build/HordeServer.Build/Jobs/JobRpcService.cs
@@ -166,9 +166,14 @@ namespace HordeServer.Jobs
 
 			RpcGetJobResponse response = new RpcGetJobResponse();
 			response.StreamId = job.StreamId.ToString();
-			response.Change = job.CommitId.GetPerforceChange();
-			response.CodeChange = job.CodeCommitId?.GetPerforceChange() ?? 0;
-			response.PreflightChange = job.PreflightCommitId?.GetPerforceChange() ?? 0;
+			response.Change = job.CommitId.TryGetPerforceChange() ?? 0;
+			response.CodeChange = job.CodeCommitId?.TryGetPerforceChange() ?? 0;
+			response.PreflightChange = job.PreflightCommitId?.TryGetPerforceChange() ?? 0;
+			response.CommitId = job.CommitId.ToString();
+			if (job.PreflightCommitId != null)
+			{
+				response.PreflightCommitId = job.PreflightCommitId.ToString();
+			}
 			response.Arguments.Add(job.Arguments);
 			return response;
 		}
@@ -234,9 +239,14 @@ namespace HordeServer.Jobs
 			response.StreamId = streamConfig.Id.ToString();
 			response.StreamName = streamConfig.Name;
 			response.EnginePath = streamConfig.EnginePath;
-			response.Change = job.CommitId.GetPerforceChange();
-			response.CodeChange = job.CodeCommitId?.GetPerforceChange() ?? 0;
-			response.PreflightChange = job.PreflightCommitId?.GetPerforceChange() ?? 0;
+			response.Change = job.CommitId.TryGetPerforceChange() ?? 0;
+			response.CodeChange = job.CodeCommitId?.TryGetPerforceChange() ?? 0;
+			response.PreflightChange = job.PreflightCommitId?.TryGetPerforceChange() ?? 0;
+			response.CommitId = job.CommitId.ToString();
+			if (job.PreflightCommitId != null)
+			{
+				response.PreflightCommitId = job.PreflightCommitId.ToString();
+			}
 			response.Arguments.AddRange(job.Arguments);
 			if (agentConfig.TempStorageDir != null)
 			{
diff --git a/Engine/Source/Programs/Horde/Plugins/Build/HordeServer.Build/Jobs/JobService.cs b/Engine/Source/Programs/Horde/Plugins/Build/HordeServer.Build/Jobs/JobService.cs
--- a/Engine/Source/Programs/Horde/Plugins/Build/HordeServer.Build/Jobs/JobService.cs
+++ b/Engine/Source/Programs/Horde/Plugins/Build/HordeServer.Build/Jobs/JobService.cs
@@ -179,7 +179,7 @@ namespace HordeServer.Jobs
 			span.SetAttribute("Name", name);
 			span.SetAttribute("Change", commitId.ToString());
 			span.SetAttribute("CodeChange", codeCommitId?.ToString());
-			span.SetAttribute("PreflightChange", options.PreflightCommitId?.GetPerforceChange());
+			span.SetAttribute("PreflightChange", options.PreflightCommitId?.ToString());
 			span.SetAttribute("StartedByUserId", options.StartedByUserId.ToString());
 			span.SetAttribute("Priority", options.Priority.ToString());
 			span.SetAttribute("ShowUgsBadges", options.ShowUgsBadges);
@@ -844,9 +844,12 @@ namespace HordeServer.Jobs
 		{
 			try
 			{
-				if (job.PreflightCommitId != null)
+				// Only check P4 shelves for numeric preflight commit IDs (P4 changelists)
+				// Hash-based preflights use commit IDs and don't have shelves to check
+				int? preflightChange = job.PreflightCommitId?.TryGetPerforceChange();
+				if (preflightChange != null)
 				{
-					(CheckShelfResult result, _) = await _perforceService.CheckShelfAsync(streamConfig, job.PreflightCommitId.GetPerforceChange(), cancellationToken);
+					(CheckShelfResult result, _) = await _perforceService.CheckShelfAsync(streamConfig, preflightChange.Value, cancellationToken);
 					if (result != CheckShelfResult.Ok)
 					{
 						_logger.LogWarning("Job {JobId} is no longer valid - check shelf returned {Result}", job.Id, result);
@@ -1028,12 +1031,17 @@ namespace HordeServer.Jobs
 							}
 							else if (job.StartedByUserId.HasValue && outcome == JobStepOutcome.Success && job.AbortedByUserId == null)
 							{
-								IUserSettings settings = await _userCollection.GetSettingsAsync(job.StartedByUserId.Value, cancellationToken);
-								if (settings.AlwaysTagPreflightCL)
+								// Only tag P4 changelists (numeric preflight commit IDs), not hash-based revisions
+								int? p4Change = job.PreflightCommitId?.TryGetPerforceChange();
+								if (p4Change != null)
 								{
-									_logger.LogInformation("Updating description for {PreflightChange} for {UserId} user settings", job.PreflightCommitId, job.StartedByUserId.Value);
-									_jobs.JobAudit(job.Id, "Updating description for {PreflightChange} for {UserId} user settings", job.PreflightCommitId, job.StartedByUserId.Value);
-									await _perforceService.UpdateChangelistDescriptionAsync(streamConfig.ClusterName, job.PreflightCommitId.GetPerforceChange(), x => x.TrimEnd() + $"\n#preflight {job.Id}", cancellationToken);
+									IUserSettings settings = await _userCollection.GetSettingsAsync(job.StartedByUserId.Value, cancellationToken);
+									if (settings.AlwaysTagPreflightCL)
+									{
+										_logger.LogInformation("Updating description for {PreflightChange} for {UserId} user settings", job.PreflightCommitId, job.StartedByUserId.Value);
+										_jobs.JobAudit(job.Id, "Updating description for {PreflightChange} for {UserId} user settings", job.PreflightCommitId, job.StartedByUserId.Value);
+										await _perforceService.UpdateChangelistDescriptionAsync(streamConfig.ClusterName, p4Change.Value, x => x.TrimEnd() + $"\n#preflight {job.Id}", cancellationToken);
+									}
 								}
 							}
 						}
@@ -1145,6 +1153,13 @@ namespace HordeServer.Jobs
 				return job;
 			}
 
+			// Auto-submit is only supported for P4 preflights (numeric commit IDs), not hash-based preflights
+			if (job.PreflightCommitId.TryGetPerforceChange() == null)
+			{
+				_logger.LogInformation("Skipping auto-submit for job {JobId} - preflight commit {CommitId} is not a P4 changelist", job.Id, job.PreflightCommitId);
+				return job;
+			}
+
 			using TelemetrySpan span = _tracer.StartActiveSpan($"{nameof(JobService)}.{nameof(AutoSubmitChangeAsync)}");
 			span.SetAttribute("Job", job.Id.ToString());
 			span.SetAttribute("Graph", graph.Id);
diff --git a/Engine/Source/Programs/Horde/Plugins/Build/HordeServer.Build/Jobs/JobTaskSource.cs b/Engine/Source/Programs/Horde/Plugins/Build/HordeServer.Build/Jobs/JobTaskSource.cs
--- a/Engine/Source/Programs/Horde/Plugins/Build/HordeServer.Build/Jobs/JobTaskSource.cs
+++ b/Engine/Source/Programs/Horde/Plugins/Build/HordeServer.Build/Jobs/JobTaskSource.cs
@@ -207,6 +207,7 @@ namespace HordeServer.Jobs
 		readonly ILogger<JobTaskSource> _logger;
 		readonly Tracer _tracer;
 		readonly ITicker _ticker;
+		readonly IEnumerable<IWorkspaceMessageEnricher> _workspaceEnrichers;
 
 		// Object used for ensuring mutual exclusion to the queues
 		readonly object _lockObject = new object();
@@ -239,7 +240,7 @@ namespace HordeServer.Jobs
 		/// <summary>
 		/// Constructor
 		/// </summary>
-		public JobTaskSource(IAclService aclService, Lazy<AgentService> agentService, IJobCollection jobs, IJobStepRefCollection jobStepRefs, IBisectTaskCollection bisectTasks, IGraphCollection graphs, IPoolCollection pools, PoolService poolService, IUgsMetadataCollection ugsMetadataCollection, IStreamCollection streamCollection, ILogCollection logCollection, PerforceLoadBalancer perforceLoadBalancer, IClock clock, IServerInfo serverInfo, IOptionsMonitor<BuildServerConfig> settings, IOptionsMonitor<BuildConfig> buildConfig, Tracer tracer, ILogger<JobTaskSource> logger)
+		public JobTaskSource(IAclService aclService, Lazy<AgentService> agentService, IJobCollection jobs, IJobStepRefCollection jobStepRefs, IBisectTaskCollection bisectTasks, IGraphCollection graphs, IPoolCollection pools, PoolService poolService, IUgsMetadataCollection ugsMetadataCollection, IStreamCollection streamCollection, ILogCollection logCollection, PerforceLoadBalancer perforceLoadBalancer, IClock clock, IServerInfo serverInfo, IOptionsMonitor<BuildServerConfig> settings, IOptionsMonitor<BuildConfig> buildConfig, Tracer tracer, ILogger<JobTaskSource> logger, IEnumerable<IWorkspaceMessageEnricher> workspaceEnrichers)
 		{
 			_aclService = aclService;
 			_agentService = agentService;
@@ -260,6 +261,7 @@ namespace HordeServer.Jobs
 			_settings = settings;
 			_tracer = tracer;
 			_logger = logger;
+			_workspaceEnrichers = workspaceEnrichers;
 
 			OnLeaseStartedProperties.Add(nameof(ExecuteJobTask.JobId), x => JobId.Parse(x.JobId)).Add(nameof(ExecuteJobTask.BatchId), x => JobStepBatchId.Parse(x.BatchId)).Add(nameof(ExecuteJobTask.LogId), x => LogId.Parse(x.LogId));
 		}
@@ -987,6 +989,29 @@ namespace HordeServer.Jobs
 
 			task.Workspace = workspaces.Last();
 
+			// For hash-based version control systems, set the revision hash on the workspace
+			if (job.CommitId.TryGetPerforceChange() == null)
+			{
+				task.Workspace.RevisionHash = job.CommitId.ToString();
+			}
+
+			// Allow plugins to enrich workspace messages (e.g., add VCS-specific credentials)
+			foreach (IWorkspaceMessageEnricher enricher in _workspaceEnrichers)
+			{
+				try
+				{
+					if (task.AutoSdkWorkspace != null && autoSdkWorkspace != null)
+					{
+						await enricher.EnrichAsync(task.AutoSdkWorkspace, autoSdkWorkspace, agent, job, cancellationToken);
+					}
+					await enricher.EnrichAsync(task.Workspace, workspace, agent, job, cancellationToken);
+				}
+				catch (Exception ex)
+				{
+					_logger.LogError(ex, "Workspace enricher {Enricher} threw an exception for job {JobId}", enricher.GetType().Name, job.Id);
+				}
+			}
+
 			return task;
 		}
 
diff --git a/Engine/Source/Programs/Horde/Plugins/Build/HordeServer.Build/Jobs/JobsController.cs b/Engine/Source/Programs/Horde/Plugins/Build/HordeServer.Build/Jobs/JobsController.cs
--- a/Engine/Source/Programs/Horde/Plugins/Build/HordeServer.Build/Jobs/JobsController.cs
+++ b/Engine/Source/Programs/Horde/Plugins/Build/HordeServer.Build/Jobs/JobsController.cs
@@ -142,31 +142,54 @@ namespace HordeServer.Jobs
 
 			// Check the preflight change is valid
 			ShelfInfo? shelfInfo = null;
+
 			if (create.PreflightCommitId != null)
 			{
-				int preflightChange = create.PreflightCommitId.GetPerforceChange();
-
-				(CheckShelfResult result, shelfInfo) = await _perforce.CheckShelfAsync(streamConfig, preflightChange, HttpContext.RequestAborted);
-				switch (result)
+				int? preflightChange = create.PreflightCommitId.TryGetPerforceChange();
+				if (preflightChange != null)
 				{
-					case CheckShelfResult.Ok:
-						break;
-					case CheckShelfResult.NoChange:
-						return BadRequest(KnownLogEvents.Horde_InvalidPreflight, "CL {Change} does not exist", preflightChange);
-					case CheckShelfResult.NoShelvedFiles:
-						return BadRequest(KnownLogEvents.Horde_InvalidPreflight, "CL {Change} does not contain any shelved files", preflightChange);
-					case CheckShelfResult.WrongStream:
-						return BadRequest(KnownLogEvents.Horde_InvalidPreflight, "CL {Change} does not contain files in {Stream}", preflightChange, streamConfig.Name);
-					case CheckShelfResult.MixedStream:
-						return BadRequest(KnownLogEvents.Horde_InvalidPreflight, "CL {Change} contains files from multiple streams", preflightChange);
-					default:
-						return BadRequest(KnownLogEvents.Horde_InvalidPreflight, "CL {Change} cannot be preflighted ({Result})", preflightChange, result);
-				}
+					// P4 path: full shelf validation
+					(CheckShelfResult result, shelfInfo) = await _perforce.CheckShelfAsync(streamConfig, preflightChange.Value, HttpContext.RequestAborted);
+					switch (result)
+					{
+						case CheckShelfResult.Ok:
+							break;
+						case CheckShelfResult.NoChange:
+							return BadRequest(KnownLogEvents.Horde_InvalidPreflight, "CL {Change} does not exist", preflightChange);
+						case CheckShelfResult.NoShelvedFiles:
+							return BadRequest(KnownLogEvents.Horde_InvalidPreflight, "CL {Change} does not contain any shelved files", preflightChange);
+						case CheckShelfResult.WrongStream:
+							return BadRequest(KnownLogEvents.Horde_InvalidPreflight, "CL {Change} does not contain files in {Stream}", preflightChange, streamConfig.Name);
+						case CheckShelfResult.MixedStream:
+							return BadRequest(KnownLogEvents.Horde_InvalidPreflight, "CL {Change} contains files from multiple streams", preflightChange);
+						default:
+							return BadRequest(KnownLogEvents.Horde_InvalidPreflight, "CL {Change} cannot be preflighted ({Result})", preflightChange, result);
+					}
 
-				if (shelfInfo!.Tags != null)
+					if (shelfInfo!.Tags != null)
+					{
+						string tagList = String.Join(";", shelfInfo.Tags.Select(x => x.ToString()));
+						environment.Add("UE_HORDE_PREFLIGHT_TAGS", tagList);
+					}
+				}
+				else
 				{
-					string tagList = String.Join(";", shelfInfo.Tags.Select(x => x.ToString()));
-					environment.Add("UE_HORDE_PREFLIGHT_TAGS", tagList);
+					// Hash-based VCS path: get commit info
+					ICommitCollection commitCollection = _commitService.GetCollection(streamConfig);
+					try
+					{
+						ICommit preflightCommit = await commitCollection.GetAsync(create.PreflightCommitId, cancellationToken);
+						shelfInfo = new ShelfInfo(preflightCommit.Description, new List<CommitTag>());
+					}
+					catch (CommitCollectionException)
+					{
+						return BadRequest(KnownLogEvents.Horde_InvalidPreflight, "Commit {CommitId} does not exist", create.PreflightCommitId);
+					}
+					catch (Exception ex)
+					{
+						_logger.LogError(ex, "Failed to look up preflight commit {CommitId}", create.PreflightCommitId);
+						return StatusCode(502, $"Failed to look up preflight commit: {ex.Message}");
+					}
 				}
 			}
 
@@ -208,12 +231,12 @@ namespace HordeServer.Jobs
 			else if (create.UpdateIssues.HasValue)
 			{
 				updateIssues = create.UpdateIssues.Value;
-			}			
+			}
 
 			// Create options for the new job
 			CreateJobOptions options = new CreateJobOptions(templateRefConfig);
 			options.PreflightCommitId = create.PreflightCommitId;
-			options.PreflightDescription = shelfInfo?.Description;			
+			options.PreflightDescription = shelfInfo?.Description;
 			options.Priority = priority;
 			options.AutoSubmit = create.AutoSubmit;
 			options.UpdateIssues = updateIssues;
diff --git a/Engine/Source/Programs/Horde/Plugins/Build/HordeServer.Build/Streams/StreamConfig.cs b/Engine/Source/Programs/Horde/Plugins/Build/HordeServer.Build/Streams/StreamConfig.cs
--- a/Engine/Source/Programs/Horde/Plugins/Build/HordeServer.Build/Streams/StreamConfig.cs
+++ b/Engine/Source/Programs/Horde/Plugins/Build/HordeServer.Build/Streams/StreamConfig.cs
@@ -118,6 +118,22 @@ namespace HordeServer.Streams
 		/// </summary>
 		public StreamId Id { get; set; }
 
+		/// <summary>
+		/// Version control service for stream commits
+		/// </summary>
+		public string VCS { get; set; } = "Perforce";
+
+		/// <summary>
+		/// Repository name for version control systems that support multiple repositories
+		/// </summary>
+		public string? RepositoryName { get; set; }
+
+		/// <summary>
+		/// Default branch name for version control systems that use branches (e.g., "main").
+		/// If not set, defaults to "main".
+		/// </summary>
+		public string? DefaultBranchName { get; set; }
+
 		/// <summary>
 		/// Direct include path for the stream config. For backwards compatibility with old config files when including from a ProjectConfig object.
 		/// </summary>
diff --git a/Engine/Source/Programs/Horde/Plugins/Build/HordeServer.Build/VersionControl/IVersionControlService.cs b/Engine/Source/Programs/Horde/Plugins/Build/HordeServer.Build/VersionControl/IVersionControlService.cs
new file mode 100644
--- /dev/null
+++ b/Engine/Source/Programs/Horde/Plugins/Build/HordeServer.Build/VersionControl/IVersionControlService.cs
@@ -0,0 +1,25 @@
+// Copyright Epic Games, Inc. All Rights Reserved.
+
+using EpicGames.Horde.Commits;
+using HordeServer.Streams;
+
+namespace HordeServer.VersionControl
+{
+	/// <summary>
+	/// Version Control interface
+	/// </summary>
+	public interface IVersionControlService
+	{
+		/// <summary>
+		/// Name of the service
+		/// </summary>
+		string Name { get; }
+
+		/// <summary>
+		/// Creates a commit source for the given stream
+		/// </summary>
+		/// <param name="streamConfig">Stream to create a commit source for</param>
+		/// <returns>Commit source instance</returns>
+		public ICommitCollection GetCommits(StreamConfig streamConfig);
+	}
+}
diff --git a/Engine/Source/Programs/Horde/Plugins/Build/HordeServer.Build/VersionControl/Perforce/IPerforceService.cs b/Engine/Source/Programs/Horde/Plugins/Build/HordeServer.Build/VersionControl/Perforce/IPerforceService.cs
--- a/Engine/Source/Programs/Horde/Plugins/Build/HordeServer.Build/VersionControl/Perforce/IPerforceService.cs
+++ b/Engine/Source/Programs/Horde/Plugins/Build/HordeServer.Build/VersionControl/Perforce/IPerforceService.cs
@@ -128,8 +128,15 @@ namespace HordeServer.VersionControl.Perforce
 	/// <summary>
 	/// Wrapper around Perforce functionality. Can use a local p4.exe client for development purposes, or a separate HordePerforceBridge instance over REST for deployments.
 	/// </summary>
-	public interface IPerforceService
+	public interface IPerforceService : IVersionControlService
 	{
+		/// <summary>
+		/// Refreshes any cached metadata for the given change
+		/// </summary>
+		/// <param name="clusterName">Name of the Perforce cluster</param>
+		/// <param name="change">Change to update</param>
+		public Task RefreshCachedCommitAsync(string clusterName, int change);
+
 		/// <summary>
 		/// Connect to the given Perforce cluster
 		/// </summary>
@@ -194,19 +201,5 @@ namespace HordeServer.VersionControl.Perforce
 		/// <param name="cancellationToken">Cancellation token for the operation</param>
 		/// <returns>Async task</returns>
 		public Task UpdateChangelistDescriptionAsync(string clusterName, int change, Func<string, string> updateFunc, CancellationToken cancellationToken = default);
-
-		/// <summary>
-		/// Refreshes any cached metadata for the given change
-		/// </summary>
-		/// <param name="clusterName">Name of the Perforce cluster</param>
-		/// <param name="change">Change to update</param>
-		public Task RefreshCachedCommitAsync(string clusterName, int change);
-
-		/// <summary>
-		/// Creates a commit source for the given stream
-		/// </summary>
-		/// <param name="streamConfig">Stream to create a commit source for</param>
-		/// <returns>Commit source instance</returns>
-		public ICommitCollection GetCommits(StreamConfig streamConfig);
 	}
 }
diff --git a/Engine/Source/Programs/Horde/Plugins/Build/HordeServer.Build/VersionControl/Perforce/PerforceService.cs b/Engine/Source/Programs/Horde/Plugins/Build/HordeServer.Build/VersionControl/Perforce/PerforceService.cs
--- a/Engine/Source/Programs/Horde/Plugins/Build/HordeServer.Build/VersionControl/Perforce/PerforceService.cs
+++ b/Engine/Source/Programs/Horde/Plugins/Build/HordeServer.Build/VersionControl/Perforce/PerforceService.cs
@@ -252,6 +252,8 @@ namespace HordeServer.VersionControl.Perforce
 			}
 		}
 
+		public string Name { get; } = "Perforce";
+
 		readonly PerforceLoadBalancer _loadBalancer;
 		readonly IOptionsMonitor<BuildConfig> _buildConfig;
 		readonly Tracer _tracer;
diff --git a/Engine/Source/Programs/Shared/EpicGames.Horde/Commits/CommitId.cs b/Engine/Source/Programs/Shared/EpicGames.Horde/Commits/CommitId.cs
--- a/Engine/Source/Programs/Shared/EpicGames.Horde/Commits/CommitId.cs
+++ b/Engine/Source/Programs/Shared/EpicGames.Horde/Commits/CommitId.cs
@@ -165,6 +165,13 @@ namespace EpicGames.Horde.Commits
 		public static new CommitIdWithOrder? FromPerforceChange(int? change)
 			=> (change == null) ? null : FromPerforceChange(change.Value);
 
+		/// <summary>
+		/// Gets the Perforce change number if the commit name is numeric, otherwise
+		/// falls back to using the Order property.
+		/// </summary>
+		public int GetPerforceChangeOrOrder()
+			=> TryGetPerforceChange() ?? Order;
+
 		/// <inheritdoc/>
 		public override bool Equals(object? obj)
 			=> Equals(obj as CommitIdWithOrder);
