Add log when test health report notification is sent

#rnx
#Synced-CL 50769868

diff --git a/Engine/Source/Programs/Horde/Plugins/Build/HordeServer.Build/Jobs/TestData/TestHealthService.cs b/Engine/Source/Programs/Horde/Plugins/Build/HordeServer.Build/Jobs/TestData/TestHealthService.cs
--- a/Engine/Source/Programs/Horde/Plugins/Build/HordeServer.Build/Jobs/TestData/TestHealthService.cs
+++ b/Engine/Source/Programs/Horde/Plugins/Build/HordeServer.Build/Jobs/TestData/TestHealthService.cs
@@ -220,6 +220,7 @@ namespace HordeServer.Jobs.TestData
 							await _notificationService.NotifyTestHealthReportAsync(report, workflowConfig.ReportChannel, userIds, cancellationToken);
 							// Update notification date or remove it if healthy
 							await _testData.UpdateTestHealthReportNotificationDateAsync(report.Id, !report.IsHealthy, cancellationToken);
+							auditLog.LogInformation("Reported state {State} to Workflow notification [{Stream}: {Workflow}] channel", report.State, streamId, workflowConfig.Id);
 						}
 					}
 				}
