[Horde][Server][Bug] When processing a single commit, exceptions can occur causing a cascading abandonment of other records.

#Synced-CL 50686579

diff --git a/Engine/Source/Programs/Horde/Plugins/Build/HordeServer.Build/VersionControl/Perforce/PerforceServiceCache.cs b/Engine/Source/Programs/Horde/Plugins/Build/HordeServer.Build/VersionControl/Perforce/PerforceServiceCache.cs
--- a/Engine/Source/Programs/Horde/Plugins/Build/HordeServer.Build/VersionControl/Perforce/PerforceServiceCache.cs
+++ b/Engine/Source/Programs/Horde/Plugins/Build/HordeServer.Build/VersionControl/Perforce/PerforceServiceCache.cs
@@ -153,7 +153,7 @@ namespace HordeServer.VersionControl.Perforce
 
 			public List<CommitTag> CommitTags { get; set; } = [];
 			public List<string> Files { get; set; } = [];
-			
+
 			[BsonElement("haf")]
 			public bool HasAllFiles { get; set; } = true;
 
@@ -192,8 +192,8 @@ namespace HordeServer.VersionControl.Perforce
 
 			public static async Task<CachedCommitDoc> FromCommitAsync(ICommit commit, CancellationToken cancellationToken)
 			{
-				List<CommitTag> commitTags = new (await commit.GetTagsAsync(cancellationToken));
-				List<string> files = new (await commit.GetFilesAsync(null, MaxFiles, cancellationToken));
+				List<CommitTag> commitTags = new(await commit.GetTagsAsync(cancellationToken));
+				List<string> files = new(await commit.GetFilesAsync(null, MaxFiles, cancellationToken));
 				bool hasAllFiles = files.Count < MaxFiles;
 				return new CachedCommitDoc(commit, commitTags, files, hasAllFiles);
 			}
@@ -215,7 +215,7 @@ namespace HordeServer.VersionControl.Perforce
 				{
 					return filter.ApplyTo(Files).Any();
 				}
-				
+
 				IReadOnlyList<string> files = await GetFilesAsync(null, null, cancellationToken);
 				return filter.ApplyTo(files).Any();
 			}
@@ -550,31 +550,40 @@ namespace HordeServer.VersionControl.Perforce
 
 						if (files.Count > 0)
 						{
-							bool hasAllFiles = describeRecord.Files.Count < MaxFiles;
+							try
+							{
+								bool hasAllFiles = describeRecord.Files.Count < MaxFiles;
 
-							ICommit commit = await CreateCommitAsync(perforce, streamInfo.StreamConfig, describeRecord, hasAllFiles, info, cancellationToken);
-							_logger.LogInformation("Replicating {StreamId} commit {Change}", streamInfo.StreamConfig.Id, describeRecord.Number);
+								ICommit commit = await CreateCommitAsync(perforce, streamInfo.StreamConfig, describeRecord, hasAllFiles, info, cancellationToken);
+								_logger.LogInformation("Replicating {StreamId} commit {Change}", streamInfo.StreamConfig.Id, describeRecord.Number);
 
-							IReadOnlyList<CommitTag> commitTags = await commit.GetTagsAsync(cancellationToken);
-							List<string> cappedFiles = files.Take(MaxFiles).ToList();
-							CachedCommitDoc commitDoc = new (commit, commitTags.ToList(), cappedFiles, hasAllFiles);
-							await AddCachedCommitAsync(commitDoc, cancellationToken);
+								IReadOnlyList<CommitTag> commitTags = await commit.GetTagsAsync(cancellationToken);
+								List<string> cappedFiles = files.Take(MaxFiles).ToList();
+								CachedCommitDoc commitDoc = new(commit, commitTags.ToList(), cappedFiles, hasAllFiles);
+								await AddCachedCommitAsync(commitDoc, cancellationToken);
 
-							await _redisService.PublishAsync(s_commitUpdateChannel, streamInfo.StreamConfig.Id, CommandFlags.FireAndForget);
+								await _redisService.PublishAsync(s_commitUpdateChannel, streamInfo.StreamConfig.Id, CommandFlags.FireAndForget);
 
-							// Emit telemetry
-							if (!streamInfo.StreamConfig.TelemetryStoreId.IsEmpty)
-							{
-								try
-								{
-									CommitTelemetry newTelemetry = await CommitTelemetry.CreateCommitTelemetryWithUserIdAsync(UserCollection, commit, CommitDescriptionTagContainer.Create(commitDoc.Description));
-									_telemetryWriter.WriteEvent(streamInfo.StreamConfig.TelemetryStoreId, newTelemetry);
-								}
-								catch (Exception ex)
+								// Emit telemetry
+								if (!streamInfo.StreamConfig.TelemetryStoreId.IsEmpty)
 								{
-									_logger.LogError(ex, "Failed to publish Commit telemetry.");
+									try
+									{
+										CommitTelemetry newTelemetry = await CommitTelemetry.CreateCommitTelemetryWithUserIdAsync(UserCollection, commit, CommitDescriptionTagContainer.Create(commitDoc.Description));
+										_telemetryWriter.WriteEvent(streamInfo.StreamConfig.TelemetryStoreId, newTelemetry);
+									}
+									catch (Exception ex)
+									{
+										_logger.LogError(ex, "Failed to publish Commit telemetry.");
+									}
 								}
 							}
+							catch (PerforceException ex)
+							{
+								// We want to make sure a single perforce exception doesn't create batch abandonment for other streams, and records.
+								// We are currently keeping this narrow to PerforceException to prevent too wide a capture, in the event there is something more systemic.
+								_logger.LogWarning(ex, "Skipping CL {Change} for stream {StreamId}: {Message}", describeRecord.Number, streamInfo.StreamConfig.Id, ex.Message);
+							}
 						}
 					}
 				}
