Horde: Make the perforce change exception raise error logging

diff --git a/Engine/Source/Programs/Horde/Plugins/Build/HordeServer.Build/VersionControl/Perforce/PerforceServiceCache.cs b/Engine/Source/Programs/Horde/Plugins/Build/HordeServer.Build/VersionControl/Perforce/PerforceServiceCache.cs
--- a/Engine/Source/Programs/Horde/Plugins/Build/HordeServer.Build/VersionControl/Perforce/PerforceServiceCache.cs
+++ b/Engine/Source/Programs/Horde/Plugins/Build/HordeServer.Build/VersionControl/Perforce/PerforceServiceCache.cs
@@ -582,7 +582,7 @@ namespace HordeServer.VersionControl.Perforce
 							{
 								// We want to make sure a single perforce exception doesn't create batch abandonment for other streams, and records.
 								// We are currently keeping this narrow to PerforceException to prevent too wide a capture, in the event there is something more systemic.
-								_logger.LogWarning(ex, "Skipping CL {Change} for stream {StreamId}: {Message}", describeRecord.Number, streamInfo.StreamConfig.Id, ex.Message);
+								_logger.LogError(ex, "Skipping CL {Change} for stream {StreamId}: {Message}", describeRecord.Number, streamInfo.StreamConfig.Id, ex.Message);
 							}
 						}
 					}
