Horde: Update release notes

#rnx

diff --git a/Engine/Source/Programs/Horde/Docs/ReleaseNotes.md b/Engine/Source/Programs/Horde/Docs/ReleaseNotes.md
--- a/Engine/Source/Programs/Horde/Docs/ReleaseNotes.md
+++ b/Engine/Source/Programs/Horde/Docs/ReleaseNotes.md
@@ -1,5 +1,26 @@
 # Release Notes
 
+## 2026-02-12
+
+* Make the perforce change exception raise error logging (50799777)
+* When processing a single commit, exceptions can occur causing a cascading abandonment of other records. (50799254)
+* Adding more support for ICommit interface (50784079)
+* Upgrade to .NET 10 SDK in BuildHorde.xml
+  Uses UE's bundled .NET SDK. (50776723)
+* Update the version of dotnet used to build to 10.0. (50767114)
+* [Horde Dashboard] Added raw `agent hours` metric to pool utilization report view and functionality to click bar chart segments for a fixed overview popup (50747397)
+* [Horde Dashboard] Adding system storage view option to agent page columns (50710586)
+* Add material-based grouping to ShaderIssueHandler
+  This ensures all failures for the same material are tracked together.
+  Fix LogShaderCompilers grouping by expanding ShaderEventMatcher to capture continuation lines and extract material paths. (50699041)
+* Expanding base context to include some hardware capabilities, and replacing DataSource with new "TestIdentity" nomenclature. (50669442)
+* Fix gRPC streaming errors in UpdateLogTail
+  The UpdateLogTail bidirectional streaming RPC was throwing unhandled exceptions under certain conditions when clients closed connections (timeout, process termination, or intentional closure). (50651383)
+* EpicGames.Oodle: Restructure to follow the standard pattern for importing native dependencies in dotnet and remove the custom assembly loader code now it's no longer required. Generate as much marshalling code at compile-time as possible. Only include the relevant platform native binaries if we're building for a specific RuntimeIdentifier. (50650488)
+* Analytics Updating key stats record. (50617535)
+* Backend library now supports cancellations. (50614189)
+* [Horde Dashboard] Fixing improper property names for fetching storage metrics within `Agents` view (50558077)
+
 ## 2026-02-05
 
 * Improve checkbox visibility in dark mode
