Horde: V3 IncrementalPoolMigration - safe hybrid mode for incremental pool migration

diff --git a/Engine/Source/Programs/Horde/HordeDashboard/src/backend/Api.ts b/Engine/Source/Programs/Horde/HordeDashboard/src/backend/Api.ts
--- a/Engine/Source/Programs/Horde/HordeDashboard/src/backend/Api.ts
+++ b/Engine/Source/Programs/Horde/HordeDashboard/src/backend/Api.ts
@@ -943,6 +943,9 @@ export type GetPoolResponse = {
 	/**Arbitrary properties associated with this event */
 	properties?: { [key: string]: string };
 
+	/** True if pool is managed via config files (read-only in dashboard). Null on pre-V3 servers. */
+	configManaged?: boolean;
+
 }
 
 /// Response describing a pool
@@ -961,6 +964,8 @@ export type GetPoolSummaryResponse = {
 	agents?: GetPoolAgentSummaryResponse[];
 	/// Utilization samples for the pool, from zero to one, with one sample for each hour
 	utilization?: number[];
+	/** True if pool is managed via config files (read-only in dashboard). Null on pre-V3 servers. */
+	configManaged?: boolean;
 }
 
 /// Numbers of agents matching various criteria
diff --git a/Engine/Source/Programs/Horde/HordeDashboard/src/components/PoolView.tsx b/Engine/Source/Programs/Horde/HordeDashboard/src/components/PoolView.tsx
--- a/Engine/Source/Programs/Horde/HordeDashboard/src/components/PoolView.tsx
+++ b/Engine/Source/Programs/Horde/HordeDashboard/src/components/PoolView.tsx
@@ -881,7 +881,7 @@ const AutoScalerPanel: React.FC = () => {
       }
    }
 
-   const disabled = !dashboard.user?.dashboardFeatures?.showPoolEditor || state.submitting;
+   const disabled = !dashboard.user?.dashboardFeatures?.showPoolEditor || state.submitting || pool.configManaged === true;
 
    let selectedColor = state.color ?? "0";
    let defaultColor = "0";
@@ -1247,8 +1247,13 @@ const PoolPanel: React.FC = () => {
                <Stack style={{ minWidth: 224 }}>
                   <Stack horizontal style={{ paddingBottom: 18 }} tokens={{ childrenGap: 8 }}>
                      <PrimaryButton text={pool.name} href={`/agents?agent=${encodeURI(pool.id)}&exact=true`} target="_blank" style={{ color: "#FFFFFF", backgroundColor: color, border: "unset", flexShrink: 1 }} />
-                     {!!dashboard.user?.dashboardFeatures?.showPoolEditor && <IconButton iconProps={{ iconName: "Edit" }} onClick={() => { setState({ ...state, showSettings: true }) }} style={{ color: "#FFFFFF", backgroundColor: color, border: "unset", flexShrink: 1 }} />}
+                     {!!dashboard.user?.dashboardFeatures?.showPoolEditor && pool.configManaged !== true && <IconButton iconProps={{ iconName: "Edit" }} onClick={() => { setState({ ...state, showSettings: true }) }} style={{ color: "#FFFFFF", backgroundColor: color, border: "unset", flexShrink: 1 }} />}
                   </Stack>
+                  {pool.configManaged === true && (
+                     <MessageBar messageBarType={MessageBarType.info} isMultiline={false} styles={{ root: { marginBottom: 8 } }}>
+                        This pool is managed via config files. To modify settings, edit pools.global.json in Perforce.
+                     </MessageBar>
+                  )}
                   <Stack>
                      <DetailsList
                         compact={true}
diff --git a/Engine/Source/Programs/Horde/HordeDashboard/src/components/PoolsView.tsx b/Engine/Source/Programs/Horde/HordeDashboard/src/components/PoolsView.tsx
--- a/Engine/Source/Programs/Horde/HordeDashboard/src/components/PoolsView.tsx
+++ b/Engine/Source/Programs/Horde/HordeDashboard/src/components/PoolsView.tsx
@@ -1,6 +1,6 @@
 // Copyright Epic Games, Inc. All Rights Reserved.
 
-import { DefaultButton, DetailsHeader, DetailsList, FontIcon, IColumn, IDetailsHeaderStyles, IDetailsListProps, ITag, Pivot, PivotItem, SelectionMode, Spinner, SpinnerSize, Stack, Sticky, StickyPositionType, TagPicker, Text, mergeStyleSets, mergeStyles } from "@fluentui/react";
+import { DefaultButton, DetailsHeader, DetailsList, FontIcon, IColumn, IDetailsHeaderStyles, IDetailsListProps, ITag, Icon, Pivot, PivotItem, SelectionMode, Spinner, SpinnerSize, Stack, Sticky, StickyPositionType, TagPicker, Text, TooltipHost, mergeStyleSets, mergeStyles } from "@fluentui/react";
 import { observer } from "mobx-react-lite";
 import React, { useEffect, useState } from "react";
 import { useNavigate, useSearchParams } from "react-router-dom";
@@ -154,7 +154,7 @@ const PoolList: React.FC = observer(() => {
       onRender: (pool: GetPoolSummaryResponse) => {
          const textColor = "white";
          const color = pool.colorValue;
-         return <Stack verticalAlign="center" verticalFill>
+         return <Stack horizontal verticalAlign="center" verticalFill tokens={{ childrenGap: 4 }}>
             <DefaultButton
                text={pool.name}
                primary
@@ -172,6 +172,11 @@ const PoolList: React.FC = observer(() => {
                   rootHovered: { border: '0px', backgroundColor: color, color: textColor, },
                   rootPressed: { border: '0px', backgroundColor: color, color: textColor, }
                }} />
+            {pool.configManaged === true && (
+               <TooltipHost content="Managed via config files (read-only)">
+                  <Icon iconName="Lock" styles={{ root: { fontSize: 12, color: statusColors.get(StatusColor.Skipped) }}} />
+               </TooltipHost>
+            )}
          </Stack>
       }
    },
diff --git a/Engine/Source/Programs/Horde/HordeDashboard/src/components/agents/AgentView.tsx b/Engine/Source/Programs/Horde/HordeDashboard/src/components/agents/AgentView.tsx
--- a/Engine/Source/Programs/Horde/HordeDashboard/src/components/agents/AgentView.tsx
+++ b/Engine/Source/Programs/Horde/HordeDashboard/src/components/agents/AgentView.tsx
@@ -882,6 +882,7 @@ type PoolEditorItem = {
    deleted: boolean;
    sortKey: string;
    numAgentsAssigned: number;
+   configManaged?: boolean;
 }
 
 // state class for the edit pools modal
@@ -930,7 +931,8 @@ class EditPoolsModalState {
             workspaces: pool.workspaces
          },
          selected: false,
-         numAgentsAssigned: this.numAgentsAssigned[pool.id] ?? 0
+         numAgentsAssigned: this.numAgentsAssigned[pool.id] ?? 0,
+         configManaged: pool.configManaged
       };
    }
 
@@ -1044,6 +1046,10 @@ class EditPoolsModalState {
       const value = { deletedPools: [] as string[], newPools: [] as string[], modifiedPools: [] as string[] };
 
       this.modifiedPools.forEach(item => {
+         // Skip config-managed pools — they cannot be modified
+         if (item.configManaged === true) {
+            return;
+         }
          if (item.deleted) {
             value.deletedPools.push(item.pool.name);
          }
@@ -1076,6 +1082,10 @@ class EditPoolsModalState {
       this._syncLastSelectedToModified();
 
       this.modifiedPools.forEach(item => {
+         // Skip config-managed pools — they cannot be modified through the API
+         if (item.configManaged === true) {
+            return;
+         }
          if (item.deleted) {
             allUpdates.push(backend.deletePool(item.pool.id));
          }
@@ -1678,6 +1688,18 @@ const PoolEditorModal: React.FC = observer(() => {
          case 'name':
             const color = item.pool.colorValue;
             const textColor = "white";
+            if (item.configManaged === true) {
+               return (
+                  <Stack styles={{ root: { height: '100%', } }} horizontal verticalAlign="center" tokens={{ childrenGap: 4 }}>
+                     <Stack.Item align={"center"}>
+                        <Text styles={{ root: { padding: '4px 8px', backgroundColor: color, color: textColor, borderRadius: 2, fontSize: '8pt', fontFamily: 'Horde Open Sans SemiBold' } }}>{item.pool.name}</Text>
+                     </Stack.Item>
+                     <TooltipHost content="Managed via config files">
+                        <Icon iconName="Lock" styles={{ root: { fontSize: 10 }}} />
+                     </TooltipHost>
+                  </Stack>
+               );
+            }
             return (
                <Stack styles={{ root: { height: '100%', } }} horizontal>
                   <Stack.Item align={"center"}>
@@ -1698,6 +1720,9 @@ const PoolEditorModal: React.FC = observer(() => {
          case 'numMachines':
             return <Stack styles={{ root: { height: '100%', } }} horizontal horizontalAlign={'center'}><Stack.Item align={"center"}>{item.numAgentsAssigned}</Stack.Item></Stack>;
          case 'del':
+            if (item.configManaged === true) {
+               return null;
+            }
             return <Stack styles={{ root: { height: '100%', } }} horizontal horizontalAlign={"center"}>
                <Stack.Item grow={0} align={"center"}>
                   <IconButton iconProps={{ iconName: 'Cancel' }} onClick={(event: any) => { editPoolsModalState.setPoolDeleted(item); }} />
diff --git a/Engine/Source/Programs/Horde/HordeServer.Shared/ConfigVersion.cs b/Engine/Source/Programs/Horde/HordeServer.Shared/ConfigVersion.cs
--- a/Engine/Source/Programs/Horde/HordeServer.Shared/ConfigVersion.cs
+++ b/Engine/Source/Programs/Horde/HordeServer.Shared/ConfigVersion.cs
@@ -26,7 +26,15 @@ namespace HordeServer
 		PoolsInConfigFiles,
 
 		/// <summary>
-		/// One after the l`ast defined version number
+		/// Safe hybrid mode for incremental pool migration. Pools are read from both config files
+		/// and the database, with deduplication. Config file pools take priority — database pools
+		/// with matching IDs are skipped. The pool editor UI remains visible for database-managed
+		/// pools. Use /api/v1/server/migrate/pool-config to export remaining database pools.
+		/// </summary>
+		IncrementalPoolMigration,
+
+		/// <summary>
+		/// One after the last defined version number
 		/// </summary>
 		LatestPlusOne,
 
diff --git a/Engine/Source/Programs/Horde/HordeServer/Configuration/ConfigService.cs b/Engine/Source/Programs/Horde/HordeServer/Configuration/ConfigService.cs
--- a/Engine/Source/Programs/Horde/HordeServer/Configuration/ConfigService.cs
+++ b/Engine/Source/Programs/Horde/HordeServer/Configuration/ConfigService.cs
@@ -600,6 +600,10 @@ namespace HordeServer.Configuration
 					{
 						message.Add($"- v{(int)ConfigVersion.PoolsInConfigFiles}: Pools should now be configured through the globals.json file rather than REST API or database. The /api/v1/server/migrate/pool-config endpoint will transcribe your configured pools into JSON.");
 					}
+					if (globalConfig.VersionEnum < ConfigVersion.IncrementalPoolMigration)
+					{
+						message.Add($"- v{(int)ConfigVersion.IncrementalPoolMigration}: Safe incremental pool migration. Pools from config files and database coexist with deduplication. Use /api/v1/server/migrate/pool-config to see remaining database pools.");
+					}
 					message.Add("");
 					message.Add($"Please migrate your installation and update the 'Version' property in globals.json to {(int)ConfigVersion.Latest}");
 
diff --git a/Engine/Source/Programs/Horde/Plugins/Compute/HordeServer.Compute/Agents/Pools/IPoolCollection.cs b/Engine/Source/Programs/Horde/Plugins/Compute/HordeServer.Compute/Agents/Pools/IPoolCollection.cs
--- a/Engine/Source/Programs/Horde/Plugins/Compute/HordeServer.Compute/Agents/Pools/IPoolCollection.cs
+++ b/Engine/Source/Programs/Horde/Plugins/Compute/HordeServer.Compute/Agents/Pools/IPoolCollection.cs
@@ -45,6 +45,12 @@ namespace HordeServer.Agents.Pools
 		/// <returns>List of pool ids</returns>
 		Task<List<PoolId>> GetPoolIdsAsync(CancellationToken cancellationToken = default);
 
+		/// <summary>
+		/// Gets a list of all pool IDs that exist in the database (PoolsV1 collection),
+		/// regardless of version settings. Used by migration tooling to detect shadow pools.
+		/// </summary>
+		Task<List<PoolId>> GetDatabasePoolIdsAsync(CancellationToken cancellationToken = default);
+
 		/// <summary>
 		/// Gets a pool by ID
 		/// </summary>
diff --git a/Engine/Source/Programs/Horde/Plugins/Compute/HordeServer.Compute/Agents/Pools/PoolCollection.cs b/Engine/Source/Programs/Horde/Plugins/Compute/HordeServer.Compute/Agents/Pools/PoolCollection.cs
--- a/Engine/Source/Programs/Horde/Plugins/Compute/HordeServer.Compute/Agents/Pools/PoolCollection.cs
+++ b/Engine/Source/Programs/Horde/Plugins/Compute/HordeServer.Compute/Agents/Pools/PoolCollection.cs
@@ -279,8 +279,17 @@ namespace HordeServer.Agents.Pools
 			pools.AddRange(computeConfig.Pools);
 			if (computeConfig.VersionEnum < ConfigVersion.PoolsInConfigFiles)
 			{
+				// V0/V1: include all DB pools (no dedup — legacy behavior preserved)
 				pools.AddRange(await _poolsV1.Find(FilterDefinition<PoolDocumentV1>.Empty).ToListAsync(cancellationToken));
 			}
+			else if (computeConfig.VersionEnum >= ConfigVersion.IncrementalPoolMigration)
+			{
+				// V3+: include DB pools, but skip any whose ID already exists in config
+				HashSet<PoolId> configPoolIds = new HashSet<PoolId>(computeConfig.Pools.Select(x => x.Id));
+				List<PoolDocumentV1> dbPools = await _poolsV1.Find(FilterDefinition<PoolDocumentV1>.Empty).ToListAsync(cancellationToken);
+				pools.AddRange(dbPools.Where(x => !configPoolIds.Contains(x.Id)));
+			}
+			// V2 (PoolsInConfigFiles): neither branch runs — config pools only
 
 			return pools;
 		}
@@ -295,8 +304,10 @@ namespace HordeServer.Agents.Pools
 			{
 				poolConfig = globalPoolConfig;
 			}
-			else if (computeConfig.VersionEnum < ConfigVersion.PoolsInConfigFiles)
+			else if (computeConfig.VersionEnum < ConfigVersion.PoolsInConfigFiles
+				|| computeConfig.VersionEnum >= ConfigVersion.IncrementalPoolMigration)
 			{
+				// V0/V1 and V3+: fall back to database if not found in config
 				poolConfig = await _poolsV1.Find(x => x.Id == id).FirstOrDefaultAsync(cancellationToken);
 			}
 
@@ -333,14 +344,37 @@ namespace HordeServer.Agents.Pools
 
 			if (computeConfig.VersionEnum < ConfigVersion.PoolsInConfigFiles)
 			{
+				// V0/V1: include all DB pool IDs (no dedup — legacy behavior preserved)
 				ProjectionDefinition<PoolDocumentV1, BsonDocument> projection = Builders<PoolDocumentV1>.Projection.Include(x => x.Id);
 				List<BsonDocument> results = await _poolsV1.Find(FilterDefinition<PoolDocumentV1>.Empty).Project(projection).ToListAsync(cancellationToken);
 				pools.AddRange(results.Select(x => new PoolId(x.GetElement("_id").Value.AsString)));
 			}
+			else if (computeConfig.VersionEnum >= ConfigVersion.IncrementalPoolMigration)
+			{
+				// V3+: include DB pool IDs, but skip any already in config
+				HashSet<PoolId> configPoolIds = new HashSet<PoolId>(pools);
+				ProjectionDefinition<PoolDocumentV1, BsonDocument> projection = Builders<PoolDocumentV1>.Projection.Include(x => x.Id);
+				List<BsonDocument> results = await _poolsV1.Find(FilterDefinition<PoolDocumentV1>.Empty).Project(projection).ToListAsync(cancellationToken);
+				pools.AddRange(results
+					.Select(x => new PoolId(x.GetElement("_id").Value.AsString))
+					.Where(id => !configPoolIds.Contains(id)));
+			}
 
 			return pools;
 		}
 
+		/// <inheritdoc/>
+		public async Task<List<PoolId>> GetDatabasePoolIdsAsync(CancellationToken cancellationToken)
+		{
+			ProjectionDefinition<PoolDocumentV1, BsonDocument> projection =
+				Builders<PoolDocumentV1>.Projection.Include(x => x.Id);
+			List<BsonDocument> results = await _poolsV1
+				.Find(FilterDefinition<PoolDocumentV1>.Empty)
+				.Project(projection)
+				.ToListAsync(cancellationToken);
+			return results.Select(x => new PoolId(x.GetElement("_id").Value.AsString)).ToList();
+		}
+
 		/// <inheritdoc/>
 		public async Task<bool> DeleteConfigAsync(PoolId id, CancellationToken cancellationToken)
 		{
diff --git a/Engine/Source/Programs/Horde/Plugins/Compute/HordeServer.Compute/Agents/Pools/PoolMessages.cs b/Engine/Source/Programs/Horde/Plugins/Compute/HordeServer.Compute/Agents/Pools/PoolMessages.cs
--- a/Engine/Source/Programs/Horde/Plugins/Compute/HordeServer.Compute/Agents/Pools/PoolMessages.cs
+++ b/Engine/Source/Programs/Horde/Plugins/Compute/HordeServer.Compute/Agents/Pools/PoolMessages.cs
@@ -532,6 +532,13 @@ namespace HordeServer.Agents.Pools
 		/// </summary>
 		public IReadOnlyDictionary<string, string> Properties { get; set; }
 
+		/// <summary>
+		/// Whether this pool is managed via config files (true) or the database (false).
+		/// Null when the distinction is not applicable (V0/V1/V2).
+		/// Config-managed pools cannot be edited through the REST API or dashboard.
+		/// </summary>
+		public bool? ConfigManaged { get; set; }
+
 		/// <summary>
 		/// Constructor
 		/// </summary>
@@ -572,7 +579,7 @@ namespace HordeServer.Agents.Pools
 	/// <param name="Stats">Counts for agents in differrent states</param>
 	/// <param name="Agents">Truncated list of agents</param>
 	/// <param name="Utilization">Utilization samples for the pool, from zero to one, with one sample for each hour</param>
-	public record class GetPoolSummaryResponse(PoolId Id, string Name, string ColorValue, bool Autoscaled, GetPoolStatsResponse? Stats, List<GetPoolAgentSummaryResponse>? Agents, List<double>? Utilization);
+	public record class GetPoolSummaryResponse(PoolId Id, string Name, string ColorValue, bool Autoscaled, GetPoolStatsResponse? Stats, List<GetPoolAgentSummaryResponse>? Agents, List<double>? Utilization, bool? ConfigManaged = null);
 
 	/// <summary>
 	/// Numbers of agents matching various criteria
diff --git a/Engine/Source/Programs/Horde/Plugins/Compute/HordeServer.Compute/Agents/Pools/PoolsController.cs b/Engine/Source/Programs/Horde/Plugins/Compute/HordeServer.Compute/Agents/Pools/PoolsController.cs
--- a/Engine/Source/Programs/Horde/Plugins/Compute/HordeServer.Compute/Agents/Pools/PoolsController.cs
+++ b/Engine/Source/Programs/Horde/Plugins/Compute/HordeServer.Compute/Agents/Pools/PoolsController.cs
@@ -54,6 +54,16 @@ namespace HordeServer.Agents.Pools
 				return Forbid(PoolAclAction.CreatePool);
 			}
 
+			// In V3+, prevent creating pools that shadow config-managed pools
+			if (_computeConfig.Value.VersionEnum >= ConfigVersion.IncrementalPoolMigration)
+			{
+				PoolId newPoolId = new PoolId(StringId.Sanitize(create.Name));
+				if (_computeConfig.Value.TryGetPool(newPoolId, out _))
+				{
+					return BadRequest($"Pool '{newPoolId}' already exists in config files and cannot be re-created through the API.");
+				}
+			}
+
 			List<PoolSizeStrategyInfo>? sizeStrategies = create.SizeStrategies?.Select(x => x.Convert()).ToList();
 			List<FleetManagerInfo>? fleetManagers = create.FleetManagers?.Select(x => x.Convert()).ToList();
 			LeaseUtilizationSettings? luSettings = create.LeaseUtilizationSettings?.Convert();
@@ -97,7 +107,8 @@ namespace HordeServer.Agents.Pools
 		[ProducesResponseType(typeof(List<GetPoolResponse>), 200)]
 		public async Task<ActionResult<List<object>>> GetPoolsAsync([FromQuery] PropertyFilter? filter = null, CancellationToken cancellationToken = default)
 		{
-			if (!_computeConfig.Value.Authorize(PoolAclAction.ListPools, User))
+			ComputeConfig computeConfig = _computeConfig.Value;
+			if (!computeConfig.Authorize(PoolAclAction.ListPools, User))
 			{
 				return Forbid(PoolAclAction.ListPools);
 			}
@@ -107,7 +118,12 @@ namespace HordeServer.Agents.Pools
 			List<object> responses = new List<object>();
 			foreach (IPoolConfig poolConfig in poolConfigs)
 			{
-				responses.Add(new GetPoolResponse(poolConfig).ApplyFilter(filter));
+				GetPoolResponse poolResponse = new GetPoolResponse(poolConfig);
+				if (computeConfig.VersionEnum >= ConfigVersion.IncrementalPoolMigration)
+				{
+					poolResponse.ConfigManaged = computeConfig.TryGetPool(poolConfig.Id, out _);
+				}
+				responses.Add(poolResponse.ApplyFilter(filter));
 			}
 			return responses;
 		}
@@ -126,7 +142,8 @@ namespace HordeServer.Agents.Pools
 		[ProducesResponseType(typeof(List<GetPoolResponse>), 200)]
 		public async Task<ActionResult<List<object>>> GetPoolSummariesAsync([FromQuery] Condition? condition = null, [FromQuery] bool stats = false, [FromQuery] int numAgents = 0, [FromQuery] int numUtilizationSamples = 0, CancellationToken cancellationToken = default)
 		{
-			if (!_computeConfig.Value.Authorize(PoolAclAction.ListPools, User))
+			ComputeConfig computeConfig = _computeConfig.Value;
+			if (!computeConfig.Authorize(PoolAclAction.ListPools, User))
 			{
 				return Forbid(PoolAclAction.ListPools);
 			}
@@ -222,7 +239,14 @@ namespace HordeServer.Agents.Pools
 						}
 					}
 
-					GetPoolSummaryResponse response = new GetPoolSummaryResponse(poolConfig.Id, poolConfig.Name, poolConfig.GetColorValue(), poolConfig.EnableAutoscaling, statsResponse, agentResponses, utilizationResponses);
+					// In V3+, indicate whether each pool is managed by config files or the database.
+					bool? configManaged = null;
+					if (computeConfig.VersionEnum >= ConfigVersion.IncrementalPoolMigration)
+					{
+						configManaged = computeConfig.TryGetPool(poolConfig.Id, out _);
+					}
+
+					GetPoolSummaryResponse response = new GetPoolSummaryResponse(poolConfig.Id, poolConfig.Name, poolConfig.GetColorValue(), poolConfig.EnableAutoscaling, statsResponse, agentResponses, utilizationResponses, configManaged);
 					responses.Add(response);
 				}
 			}
@@ -253,7 +277,12 @@ namespace HordeServer.Agents.Pools
 				return NotFound(poolIdValue);
 			}
 
-			return new GetPoolResponse(pool).ApplyFilter(filter);
+			GetPoolResponse response = new GetPoolResponse(pool);
+			if (_computeConfig.Value.VersionEnum >= ConfigVersion.IncrementalPoolMigration)
+			{
+				response.ConfigManaged = _computeConfig.Value.TryGetPool(poolIdValue, out _);
+			}
+			return response.ApplyFilter(filter);
 		}
 
 		/// <summary>
@@ -274,6 +303,13 @@ namespace HordeServer.Agents.Pools
 
 			PoolId poolIdValue = new PoolId(poolId);
 
+			// In V3+, config-managed pools cannot be modified through the API
+			if (_computeConfig.Value.VersionEnum >= ConfigVersion.IncrementalPoolMigration
+				&& _computeConfig.Value.TryGetPool(poolIdValue, out _))
+			{
+				return BadRequest($"Pool '{poolId}' is managed via config files and cannot be modified through the API. Edit pools.global.json in Perforce instead.");
+			}
+
 			List<PoolSizeStrategyInfo>? newSizeStrategies = update.SizeStrategies?.Select(x => x.Convert()).ToList();
 			List<FleetManagerInfo>? newFleetManagers = update.FleetManagers?.Select(x => x.Convert()).ToList();
 			TimeSpan? conformInterval = update.ConformInterval == null ? null : TimeSpan.FromHours(update.ConformInterval.Value);
@@ -318,6 +354,13 @@ namespace HordeServer.Agents.Pools
 			}
 
 			PoolId poolIdValue = new PoolId(poolId);
+
+			if (_computeConfig.Value.VersionEnum >= ConfigVersion.IncrementalPoolMigration
+				&& _computeConfig.Value.TryGetPool(poolIdValue, out _))
+			{
+				return BadRequest($"Pool '{poolId}' is managed via config files and cannot be deleted through the API. Edit pools.global.json in Perforce instead.");
+			}
+
 			if (!await _poolCollection.DeleteConfigAsync(poolIdValue))
 			{
 				return NotFound(poolIdValue);
@@ -340,6 +383,18 @@ namespace HordeServer.Agents.Pools
 				return Forbid(PoolAclAction.UpdatePool);
 			}
 
+			if (_computeConfig.Value.VersionEnum >= ConfigVersion.IncrementalPoolMigration)
+			{
+				List<string> configManagedIds = batchUpdates
+					.Where(u => _computeConfig.Value.TryGetPool(new PoolId(u.Id), out _))
+					.Select(u => u.Id)
+					.ToList();
+				if (configManagedIds.Count > 0)
+				{
+					return BadRequest($"The following pools are config-managed and cannot be modified via API: {String.Join(", ", configManagedIds)}");
+				}
+			}
+
 			foreach (BatchUpdatePoolRequest update in batchUpdates)
 			{
 				PoolId poolIdValue = new PoolId(update.Id);
diff --git a/Engine/Source/Programs/Horde/Plugins/Compute/HordeServer.Compute/ComputeResponseFilter.cs b/Engine/Source/Programs/Horde/Plugins/Compute/HordeServer.Compute/ComputeResponseFilter.cs
--- a/Engine/Source/Programs/Horde/Plugins/Compute/HordeServer.Compute/ComputeResponseFilter.cs
+++ b/Engine/Source/Programs/Horde/Plugins/Compute/HordeServer.Compute/ComputeResponseFilter.cs
@@ -23,7 +23,15 @@ namespace HordeServer
 			if (response is GetDashboardFeaturesResponse featuresResponse)
 			{
 				ComputeConfig computeConfig = _computeConfig.CurrentValue;
-				featuresResponse.ShowPoolEditor = computeConfig.VersionEnum < ConfigVersion.PoolsInConfigFiles && (computeConfig.Authorize(PoolAclAction.CreatePool, context.User) || computeConfig.Authorize(PoolAclAction.UpdatePool, context.User));
+
+				// Pool editor is available in V0/V1 (legacy) and V3+ (incremental migration).
+				// In V2 (PoolsInConfigFiles), the DB is not read so the editor is hidden.
+				bool hasDbPools = computeConfig.VersionEnum < ConfigVersion.PoolsInConfigFiles
+					|| computeConfig.VersionEnum >= ConfigVersion.IncrementalPoolMigration;
+				featuresResponse.ShowPoolEditor = hasDbPools
+					&& (computeConfig.Authorize(PoolAclAction.CreatePool, context.User)
+						|| computeConfig.Authorize(PoolAclAction.UpdatePool, context.User));
+
 				featuresResponse.ShowRemoteDesktop = computeConfig.Authorize(AgentAclAction.UpdateAgent, context.User);
 			}
 		}
diff --git a/Engine/Source/Programs/Horde/Plugins/Compute/HordeServer.Compute/Server/PoolsServerController.cs b/Engine/Source/Programs/Horde/Plugins/Compute/HordeServer.Compute/Server/PoolsServerController.cs
--- a/Engine/Source/Programs/Horde/Plugins/Compute/HordeServer.Compute/Server/PoolsServerController.cs
+++ b/Engine/Source/Programs/Horde/Plugins/Compute/HordeServer.Compute/Server/PoolsServerController.cs
@@ -7,6 +7,7 @@ using HordeServer.Utilities;
 using Microsoft.AspNetCore.Authorization;
 using Microsoft.AspNetCore.Mvc;
 using Microsoft.Extensions.DependencyInjection;
+using Microsoft.Extensions.Logging;
 using Microsoft.Extensions.Options;
 
 namespace HordeServer.Server
@@ -21,14 +22,16 @@ namespace HordeServer.Server
 	{
 		readonly IServiceProvider _serviceProvider;
 		readonly IOptionsSnapshot<ComputeConfig> _computeConfig;
+		readonly ILogger<PoolsServerController> _logger;
 
 		/// <summary>
 		/// Constructor
 		/// </summary>
-		public PoolsServerController(IServiceProvider serviceProvider, IOptionsSnapshot<ComputeConfig> computeConfig)
+		public PoolsServerController(IServiceProvider serviceProvider, IOptionsSnapshot<ComputeConfig> computeConfig, ILogger<PoolsServerController> logger)
 		{
 			_serviceProvider = serviceProvider;
 			_computeConfig = computeConfig;
+			_logger = logger;
 		}
 
 		/// <summary>
@@ -111,7 +114,52 @@ namespace HordeServer.Server
 				configs.Add(config);
 			}
 
-			return new { Pools = configs };
+			var response = new { Pools = configs };
+			if (_computeConfig.Value.VersionEnum >= ConfigVersion.IncrementalPoolMigration)
+			{
+				int configPoolCount = _computeConfig.Value.Pools.Count;
+				_logger.LogInformation(
+					"Pool migration status: {ConfigPools} in config, {DbPools} remaining in database",
+					configPoolCount, response.Pools.Count);
+			}
+			return response;
+		}
+
+		/// <summary>
+		/// Returns pool migration status: how many pools are in config, DB-only, and shadowed.
+		/// </summary>
+		[HttpGet]
+		[Route("/api/v1/server/migrate/pool-status")]
+		public async Task<ActionResult<object>> GetPoolMigrationStatusAsync(CancellationToken cancellationToken)
+		{
+			if (!_computeConfig.Value.Authorize(PoolAclAction.ListPools, User))
+			{
+				return Forbid(PoolAclAction.ListPools);
+			}
+
+			ComputeConfig computeConfig = _computeConfig.Value;
+			HashSet<PoolId> configPoolIds = new HashSet<PoolId>(computeConfig.Pools.Select(x => x.Id));
+
+			IPoolCollection poolCollection = _serviceProvider.GetRequiredService<IPoolCollection>();
+			List<PoolId> dbPoolIdList = await poolCollection.GetDatabasePoolIdsAsync(cancellationToken);
+			HashSet<PoolId> dbPoolIds = new HashSet<PoolId>(dbPoolIdList);
+
+			// Classify
+			HashSet<PoolId> shadowIds = new HashSet<PoolId>(configPoolIds);
+			shadowIds.IntersectWith(dbPoolIds);
+
+			HashSet<PoolId> dbOnlyIds = new HashSet<PoolId>(dbPoolIds);
+			dbOnlyIds.ExceptWith(configPoolIds);
+
+			return new
+			{
+				Version = (int)computeConfig.VersionEnum,
+				ConfigPools = configPoolIds.Count,
+				DbOnlyPools = dbOnlyIds.Count,
+				ShadowPools = shadowIds.Count,
+				ShadowPoolIds = shadowIds.Select(x => x.ToString()).OrderBy(x => x).ToList(),
+				DbOnlyPoolIds = dbOnlyIds.Select(x => x.ToString()).OrderBy(x => x).ToList()
+			};
 		}
 	}
 }
