Horde: Fix doc string
#rnx

diff --git a/Engine/Source/Programs/Horde/Plugins/Compute/HordeServer.Compute/Agents/Pools/PoolMessages.cs b/Engine/Source/Programs/Horde/Plugins/Compute/HordeServer.Compute/Agents/Pools/PoolMessages.cs
--- a/Engine/Source/Programs/Horde/Plugins/Compute/HordeServer.Compute/Agents/Pools/PoolMessages.cs
+++ b/Engine/Source/Programs/Horde/Plugins/Compute/HordeServer.Compute/Agents/Pools/PoolMessages.cs
@@ -579,6 +579,7 @@ namespace HordeServer.Agents.Pools
 	/// <param name="Stats">Counts for agents in differrent states</param>
 	/// <param name="Agents">Truncated list of agents</param>
 	/// <param name="Utilization">Utilization samples for the pool, from zero to one, with one sample for each hour</param>
+	/// <param name="ConfigManaged">Whether this pool is managed by configuration</param>
 	public record class GetPoolSummaryResponse(PoolId Id, string Name, string ColorValue, bool Autoscaled, GetPoolStatsResponse? Stats, List<GetPoolAgentSummaryResponse>? Agents, List<double>? Utilization, bool? ConfigManaged = null);
 
 	/// <summary>
