HordeServer: Update Dockerfile to copy dependencies before nuget restore rather than after
Add additional excluded directories to ignore file

#rnx

#Synced-CL 50830039

diff --git a/Engine/Source/Programs/Horde/HordeServer/Dockerfile b/Engine/Source/Programs/Horde/HordeServer/Dockerfile
--- a/Engine/Source/Programs/Horde/HordeServer/Dockerfile
+++ b/Engine/Source/Programs/Horde/HordeServer/Dockerfile
@@ -28,42 +28,30 @@ RUN wget https://repo.mongodb.org/apt/debian/dists/bookworm/mongodb-org/7.0/main
 	dpkg -x mongodb-org-server_7.0.4_amd64.deb /tmp/mongodb && \
 	cp /tmp/mongodb/usr/bin/mongod /usr/local/bin/mongod
 
-# Restore all the C# projects
-COPY Source/Programs/Shared/EpicGames.Core/*.csproj ./Source/Programs/Shared/EpicGames.Core/
-COPY Source/Programs/Shared/EpicGames.Horde/*.csproj ./Source/Programs/Shared/EpicGames.Horde/
-COPY Source/Programs/Shared/EpicGames.IoHash/*.csproj ./Source/Programs/Shared/EpicGames.IoHash/
-COPY Source/Programs/Shared/EpicGames.LogMatchers/*.csproj ./Source/Programs/Shared/EpicGames.LogMatchers/
-COPY Source/Programs/Shared/EpicGames.Perforce/*.csproj ./Source/Programs/Shared/EpicGames.Perforce/
-COPY Source/Programs/Shared/EpicGames.Perforce.Managed/*.csproj ./Source/Programs/Shared/EpicGames.Perforce.Managed/
-COPY Source/Programs/Shared/EpicGames.Redis/*.csproj ./Source/Programs/Shared/EpicGames.Redis/
-COPY Source/Programs/Shared/EpicGames.Serialization/*.csproj ./Source/Programs/Shared/EpicGames.Serialization/
-COPY Source/Programs/Shared/EpicGames.Serialization.Tests/*.csproj ./Source/Programs/Shared/EpicGames.Serialization.Tests/
-COPY Source/Programs/Horde/Drivers/JobDriver/*.csproj ./Source/Programs/Horde/Drivers/JobDriver/
-COPY Source/Programs/Horde/HordeAgent/*.csproj ./Source/Programs/Horde/HordeAgent/
-COPY Source/Programs/Horde/HordeServer/*.csproj ./Source/Programs/Horde/HordeServer/
-COPY Source/Programs/Horde/HordeServer.Shared/*.csproj ./Source/Programs/Horde/HordeServer.Shared/
-COPY Source/Programs/Horde/HordeServer.Tests/*.csproj ./Source/Programs/Horde/HordeServer.Tests/
-RUN dotnet restore Source/Programs/Horde/HordeServer/HordeServer.csproj -p:WithHordeStorage=false
-
-# Build the server
+# Copy build dependencies
 COPY Binaries/DotNET/EpicGames.Perforce.Native/ Binaries/DotNET/EpicGames.Perforce.Native/
-COPY Source/Programs/Shared/ Source/Programs/Shared/
+COPY Source/Programs/Shared/ ./Source/Programs/Shared/
 COPY Source/Programs/Horde/.editorconfig Source/Programs/Horde/.editorconfig
 COPY Source/Programs/Horde/Plugins/ Source/Programs/Horde/Plugins/
 COPY Source/Programs/Horde/HordeServer/ Source/Programs/Horde/HordeServer/
 COPY Source/Programs/Horde/HordeServer.Shared/ Source/Programs/Horde/HordeServer.Shared/
+
+# Restore all the C# projects
+RUN dotnet restore Source/Programs/Horde/HordeServer/HordeServer.csproj -p:WithHordeStorage=false
+
+# Build the server
 ARG msbuild_args
 RUN dotnet publish Source/Programs/Horde/HordeServer/HordeServer.csproj -c Release -o out -p:WithHordeStorage=false $msbuild_args
 
-# Run tests
+# Copy test dependencies
 COPY Source/Programs/Horde/Drivers/JobDriver/ Source/Programs/Horde/Drivers/JobDriver/
 COPY Source/Programs/Horde/HordeAgent/ Source/Programs/Horde/HordeAgent/
 COPY Source/Programs/Horde/HordeServer.Tests/ Source/Programs/Horde/HordeServer.Tests/
 COPY Source/Programs/Horde/HordeServer.Tests.Shared/ Source/Programs/Horde/HordeServer.Tests.Shared/
-COPY Source/Programs/Shared/EpicGames.Serialization.Tests/ Source/Programs/Shared/EpicGames.Serialization.Tests/
+COPY Source/Programs/Horde/Scripts/test.sh Source/Programs/Horde/Scripts/test.sh
 
+# Run tests
 ARG code_coverage
-COPY Source/Programs/Horde/Scripts/test.sh Source/Programs/Horde/Scripts/test.sh
 RUN bash Source/Programs/Horde/Scripts/test.sh
 
 # Remove native libs not used on Linux x86_64
diff --git a/Engine/Source/Programs/Horde/HordeServer/Dockerfile.dockerignore b/Engine/Source/Programs/Horde/HordeServer/Dockerfile.dockerignore
--- a/Engine/Source/Programs/Horde/HordeServer/Dockerfile.dockerignore
+++ b/Engine/Source/Programs/Horde/HordeServer/Dockerfile.dockerignore
@@ -17,6 +17,11 @@
 # Ensure bin/obj files generated by .NET builds on the host are ignored
 **/bin
 **/obj
+**/uatbin
+**/uatobj
+
+# Ignore Visual Studio directory
+**/.vs
 
 # Exclude local config files
 appsettings.Local.json
