Horde: Fix flaky AgentRelayService tests by isolating Redis pub/sub channels

* Redis pub/sub channels are server-wide and not scoped to database numbers, causing cross-test interference when tests run in parallel
* Make the Redis channel name configurable and set it to a unique value per test instance to prevent race conditions where one test's updates would wake up another test's long-polling clients

diff --git a/Engine/Source/Programs/Horde/Plugins/Build/HordeServer.Build.Tests/Agents/Relay/AgentRelayServiceTests.cs b/Engine/Source/Programs/Horde/Plugins/Build/HordeServer.Build.Tests/Agents/Relay/AgentRelayServiceTests.cs
--- a/Engine/Source/Programs/Horde/Plugins/Build/HordeServer.Build.Tests/Agents/Relay/AgentRelayServiceTests.cs
+++ b/Engine/Source/Programs/Horde/Plugins/Build/HordeServer.Build.Tests/Agents/Relay/AgentRelayServiceTests.cs
@@ -52,6 +52,7 @@ public class AgentRelayServiceTests : BuildTestSetup
 	{
 		_service = new AgentRelayService(GetRedisServiceSingleton(), Clock, Tracer, NullLogger<AgentRelayService>.Instance);
 		_service.SetRandomSeed(1);
+		_service.SetRedisChannel($"relay/update/{Guid.NewGuid()}");
 	}
 
 	public override async ValueTask DisposeAsync()
diff --git a/Engine/Source/Programs/Horde/Plugins/Compute/HordeServer.Compute/Agents/Relay/AgentRelayService.cs b/Engine/Source/Programs/Horde/Plugins/Compute/HordeServer.Compute/Agents/Relay/AgentRelayService.cs
--- a/Engine/Source/Programs/Horde/Plugins/Compute/HordeServer.Compute/Agents/Relay/AgentRelayService.cs
+++ b/Engine/Source/Programs/Horde/Plugins/Compute/HordeServer.Compute/Agents/Relay/AgentRelayService.cs
@@ -72,7 +72,7 @@ public sealed class AgentRelayService : RelayRpc.RelayRpcBase, IHostedService, I
 {
 	private record PortMappingsInfo(int Revision, List<PortMapping> PortMappings);
 
-	private const string RedisChannelUpdate = "relay/update";
+	private string _redisChannelUpdate = "relay/update";
 	private static string KeyClusters() => "relay/clusters";
 	private static string KeyPortMappings(ClusterId clusterId) => $"relay/port-mappings/{clusterId}";
 	private static string KeyPortMappingRevision(ClusterId clusterId) => $"relay/port-mappings-revision/{clusterId}";
@@ -359,12 +359,12 @@ public sealed class AgentRelayService : RelayRpc.RelayRpcBase, IHostedService, I
 
 	private Task PublishUpdateEventAsync(string clusterId)
 	{
-		return _redis.GetDatabase().PublishAsync(RedisChannel.Literal(RedisChannelUpdate), clusterId);
+		return _redis.GetDatabase().PublishAsync(RedisChannel.Literal(_redisChannelUpdate), clusterId);
 	}
 
 	private async Task<IAsyncDisposable> SubscribeToUpdateEventAsync(Action<string> onUpdate)
 	{
-		return await _redis.GetDatabase().Multiplexer.SubscribeAsync(RedisChannelUpdate, onUpdate);
+		return await _redis.GetDatabase().Multiplexer.SubscribeAsync(_redisChannelUpdate, onUpdate);
 	}
 
 	private void OnPortMappingUpdate(string clusterIdStr)
@@ -469,6 +469,11 @@ public sealed class AgentRelayService : RelayRpc.RelayRpcBase, IHostedService, I
 		_random = new Random(seed);
 	}
 
+	internal void SetRedisChannel(string channel)
+	{
+		_redisChannelUpdate = channel;
+	}
+
 	internal void SetTestDelayBeforeTaskCapture(Func<Task>? delay)
 	{
 		_testDelayBeforeTaskCapture = delay;
