Horde: Add OptionalRequires support for BuildGraph nodes

Allow BuildGraph nodes to declare dependencies using the new OptionalRequires attribute.
The optional dependencies are not required to be successful for a node's execution; if an optional dependency fails the node can still run and the node's execution will decide how missing inputs are handled instead of the graph. Any optional dependencies, like required dependencies, will be scheduled to run to completion before a node can start.

The change is the Horde side changes to support optional requires.

#Synced-CL 0

diff --git a/Engine/Source/Programs/Horde/Drivers/JobDriver/Execution/JobExecutor.cs b/Engine/Source/Programs/Horde/Drivers/JobDriver/Execution/JobExecutor.cs
--- a/Engine/Source/Programs/Horde/Drivers/JobDriver/Execution/JobExecutor.cs
+++ b/Engine/Source/Programs/Horde/Drivers/JobDriver/Execution/JobExecutor.cs
@@ -39,6 +39,7 @@ namespace JobDriver.Execution
 	/// <param name="EnvVars">Environment variables for the step</param>
 	/// <param name="Warnings">Whether warnings are enabled</param>
 	/// <param name="Inputs">Input dependencies for the step</param>
+	/// <param name="OptionalInputs">Optional input dependencies for the step</param>
 	/// <param name="OutputNames">Names of outputs produced by the step</param>
 	/// <param name="PublishOutputs">Indices of outputs to publish</param>
 	/// <param name="Artifacts">Artifact definitions for the step</param>
@@ -52,6 +53,7 @@ namespace JobDriver.Execution
 		IReadOnlyDictionary<string, string> EnvVars,
 		bool? Warnings,
 		IReadOnlyList<string> Inputs,
+		IReadOnlyList<string> OptionalInputs,
 		IReadOnlyList<string> OutputNames,
 		IList<int> PublishOutputs,
 		IReadOnlyList<RpcCreateGraphArtifactRequest> Artifacts
@@ -61,7 +63,7 @@ namespace JobDriver.Execution
 		/// Constructor from an RPC response
 		/// </summary>
 		public JobStepInfo(RpcBeginStepResponse response)
-			: this(JobStepId.Parse(response.StepId), LogId.Parse(response.LogId), response.Name, response.Credentials, response.Properties, response.EnvVars, response.Warnings, response.Inputs, response.OutputNames, response.PublishOutputs, response.Artifacts)
+			: this(JobStepId.Parse(response.StepId), LogId.Parse(response.LogId), response.Name, response.Credentials, response.Properties, response.EnvVars, response.Warnings, response.Inputs, response.OptionalInputs, response.OutputNames, response.PublishOutputs, response.Artifacts)
 		{
 		}
 	}
@@ -134,8 +136,10 @@ namespace JobDriver.Execution
 			public bool? AllowRetry { get; set; }
 			public bool? Warnings { get; set; }
 			public List<string> Inputs { get; set; } = new List<string>();
+			public List<string> OptionalInputs { get; set; } = [];
 			public List<string> Outputs { get; set; } = new List<string>();
 			public List<string> InputDependencies { get; set; } = new List<string>();
+			public List<string> OptionalInputDependencies { get; set; } = [];
 			public List<string> OrderDependencies { get; set; } = new List<string>();
 			public Dictionary<string, string> Annotations { get; set; } = new Dictionary<string, string>(StringComparer.OrdinalIgnoreCase);
 		}
@@ -1150,6 +1154,10 @@ namespace JobDriver.Execution
 					{
 						createNode.Inputs.Add(exportedNode.Inputs);
 					}
+					if (exportedNode.OptionalInputs != null)
+					{
+						createNode.OptionalInputs.Add(exportedNode.OptionalInputs);
+					}
 					if (exportedNode.Outputs != null)
 					{
 						createNode.Outputs.Add(exportedNode.Outputs);
@@ -1158,6 +1166,10 @@ namespace JobDriver.Execution
 					{
 						createNode.InputDependencies.Add(exportedNode.InputDependencies);
 					}
+					if (exportedNode.OptionalInputDependencies != null)
+					{
+						createNode.OptionalInputDependencies.Add(exportedNode.OptionalInputDependencies);
+					}
 					if (exportedNode.OrderDependencies != null)
 					{
 						createNode.OrderDependencies.Add(exportedNode.OrderDependencies);
@@ -1438,6 +1450,35 @@ namespace JobDriver.Execution
 				inputStorageBlocks.UnionWith(fileList.Blocks);
 			}
 
+			HashSet<string> optionalInputs = new(step.OptionalInputs);
+			foreach (string optionalInput in optionalInputs)
+			{
+				int slashIdx = optionalInput.IndexOf('/', StringComparison.Ordinal);
+				if (slashIdx == -1)
+				{
+					// Bad data (malformed format) is a programming error and should fail the step,
+					// unlike retrieval failures which are expected when the producing step has failed.
+					logger.LogError("Missing slash from node optional input: {OptionalInput}", optionalInput);
+					return false;
+				}
+
+				string nodeName = optionalInput.Substring(0, slashIdx);
+				string tagName = optionalInput.Substring(slashIdx + 1);
+
+				TempStorageTagManifest fileList;
+				try
+				{
+					fileList = await TempStorage.RetrieveTagAsync(HordeClient, JobId, step.StepId, nodeName, tagName, manifestDir, logger, cancellationToken);
+				}
+				catch (Exception ex) when (ex is not OperationCanceledException)
+				{
+					logger.LogWarning(ex, "Failed to read optional input {OptionalInput} from temp storage: {Message}. The producing step may have failed. Continuing without it", optionalInput, ex.Message);
+					continue;
+				}
+				tagNameToFileSet[tagName] = fileList.ToFileSet(unrealRootDir);
+				inputStorageBlocks.UnionWith(fileList.Blocks);
+			}
+
 			// Read the manifests for all the input storage blocks
 			Dictionary<TempStorageBlockRef, TempStorageBlockManifest> inputManifests = new Dictionary<TempStorageBlockRef, TempStorageBlockManifest>();
 
@@ -2272,7 +2313,10 @@ namespace JobDriver.Execution
 
 				await JobRpc.UpdateGraphAsync(updateGraph, null, null, cancellationToken);
 
-				HashSet<string> publishOutputNames = new HashSet<string>(updateGraph.Groups.SelectMany(x => x.Nodes).SelectMany(x => x.Inputs), StringComparer.OrdinalIgnoreCase);
+				IEnumerable<RpcCreateNodeRequest> allNodes = updateGraph.Groups.SelectMany(x => x.Nodes);
+				HashSet<string> publishOutputNames = new HashSet<string>(StringComparer.OrdinalIgnoreCase);
+				publishOutputNames.UnionWith(allNodes.SelectMany(x => x.Inputs));
+				publishOutputNames.UnionWith(allNodes.SelectMany(x => x.OptionalInputs));
 				foreach (string publishOutputName in publishOutputNames)
 				{
 					jobLogger.LogInformation("Required output: {OutputName}", publishOutputName);
diff --git a/Engine/Source/Programs/Horde/HordeDashboard/src/backend/Api.ts b/Engine/Source/Programs/Horde/HordeDashboard/src/backend/Api.ts
--- a/Engine/Source/Programs/Horde/HordeDashboard/src/backend/Api.ts
+++ b/Engine/Source/Programs/Horde/HordeDashboard/src/backend/Api.ts
@@ -1478,6 +1478,9 @@ export type GetStepResponse = {
 	/// Indices of nodes which must have succeeded for this node to run
 	inputDependencies?: JobStepId[];
 
+	/// Indices of nodes which may have succeeded for this node to run
+	optionalInputDependencies?: JobStepId[];
+
 	/// Indices of nodes which must have completed for this node to run
 	orderDependencies?: JobStepId[];
 
diff --git a/Engine/Source/Programs/Horde/HordeDashboard/src/components/jobDetailsV2/JobDetailViewSteps.tsx b/Engine/Source/Programs/Horde/HordeDashboard/src/components/jobDetailsV2/JobDetailViewSteps.tsx
--- a/Engine/Source/Programs/Horde/HordeDashboard/src/components/jobDetailsV2/JobDetailViewSteps.tsx
+++ b/Engine/Source/Programs/Horde/HordeDashboard/src/components/jobDetailsV2/JobDetailViewSteps.tsx
@@ -445,7 +445,7 @@ const StepsPanelInner: React.FC<{ jobDetails: JobDetailsV2, depStepId?: string }
          if (!rstep) {
             return;
          }
-         [rstep.inputDependencies, rstep.orderDependencies].flat().forEach(id => {
+         [rstep.inputDependencies, rstep.optionalInputDependencies, rstep.orderDependencies].flat().forEach(id => {
             const s = jobDetails.stepById(id);
             if (s && !stepFilter.find(s => s.id === id)) {
                stepFilter.push(s);
diff --git a/Engine/Source/Programs/Horde/HordeDashboard/src/components/jobDetailsV2/JobDetailsViewCommon.tsx b/Engine/Source/Programs/Horde/HordeDashboard/src/components/jobDetailsV2/JobDetailsViewCommon.tsx
--- a/Engine/Source/Programs/Horde/HordeDashboard/src/components/jobDetailsV2/JobDetailsViewCommon.tsx
+++ b/Engine/Source/Programs/Horde/HordeDashboard/src/components/jobDetailsV2/JobDetailsViewCommon.tsx
@@ -723,7 +723,7 @@ export class JobDetailsV2 extends PollBase {
             return;
          }
 
-         [step.inputDependencies, step.orderDependencies].flat().forEach(id => {
+         [step.inputDependencies, step.optionalInputDependencies, step.orderDependencies].flat().forEach(id => {
             const s = this.stepById(id);
             if (s && !steps.find(s => s.id === id)) {
                steps.push(s);
diff --git a/Engine/Source/Programs/Horde/Plugins/Build/HordeServer.Build.Tests/Jobs/GraphCollectionTests.cs b/Engine/Source/Programs/Horde/Plugins/Build/HordeServer.Build.Tests/Jobs/GraphCollectionTests.cs
new file mode 100644
--- /dev/null
+++ b/Engine/Source/Programs/Horde/Plugins/Build/HordeServer.Build.Tests/Jobs/GraphCollectionTests.cs
@@ -0,0 +1,63 @@
+// Copyright Epic Games, Inc. All Rights Reserved.
+
+using EpicGames.Core;
+using EpicGames.Horde.Jobs.Graphs;
+using MongoDB.Bson;
+using MongoDB.Driver;
+
+namespace HordeServer.Tests.Jobs
+{
+	[TestClass]
+	public class GraphCollectionTests : BuildTestSetup
+	{
+		[TestMethod]
+		public async Task TestLegacyNodeWithoutOptionalInputsAsync()
+		{
+			ContentHash hash = ContentHash.MD5("legacy-graph-test");
+
+			// Build a BsonDocument simulating a pre-OptionalInput graph document.
+			// The Node deliberately omits nullable fields to mirror what older documents look like in MongoDB.
+			BsonDocument nodeDoc = new()
+			{
+				{ "Name", "TestNode" },
+				{ "InputDependencies", new BsonArray() },
+				{ "OrderDependencies", new BsonArray() },
+				{ "Priority", 0 },
+				{ "AllowRetry", true },
+				{ "RunEarly", false },
+				{ "Warnings", true }
+			};
+
+			BsonDocument groupDoc = new()
+			{
+				{ "AgentType", "Win64" },
+				{ "Nodes", new BsonArray { nodeDoc } }
+			};
+
+			BsonDocument graphDoc = new()
+			{
+				{ "_id", hash.ToString() },
+				{ "Schema", 0 },
+				{ "Groups", new BsonArray { groupDoc } },
+				{ "Aggregates", new BsonArray() },
+				{ "Labels", new BsonArray() },
+				{ "Artifacts", new BsonArray() }
+			};
+
+			IMongoCollection<BsonDocument> collection = MongoService.GetCollection<BsonDocument>("Graphs");
+			await collection.InsertOneAsync(graphDoc);
+
+			IGraph graph = await GraphCollection.GetAsync(hash, CancellationToken.None);
+
+			Assert.AreEqual(1, graph.Groups.Count);
+			Assert.AreEqual(1, graph.Groups[0].Nodes.Count);
+			Assert.AreEqual("TestNode", graph.Groups[0].Nodes[0].Name);
+
+			Assert.IsNotNull(graph.Groups[0].Nodes[0].OptionalInputs);
+			Assert.AreEqual(0, graph.Groups[0].Nodes[0].OptionalInputs.Count);
+
+			Assert.IsNotNull(graph.Groups[0].Nodes[0].OptionalInputDependencies);
+			Assert.AreEqual(0, graph.Groups[0].Nodes[0].OptionalInputDependencies.Length);
+		}
+	}
+}
diff --git a/Engine/Source/Programs/Horde/Plugins/Build/HordeServer.Build.Tests/Jobs/JobCollectionTests.cs b/Engine/Source/Programs/Horde/Plugins/Build/HordeServer.Build.Tests/Jobs/JobCollectionTests.cs
--- a/Engine/Source/Programs/Horde/Plugins/Build/HordeServer.Build.Tests/Jobs/JobCollectionTests.cs
+++ b/Engine/Source/Programs/Horde/Plugins/Build/HordeServer.Build.Tests/Jobs/JobCollectionTests.cs
@@ -88,7 +88,22 @@ namespace HordeServer.Tests.Jobs
 
 		static NewNode AddNode(NewGroup group, string name, string[]? inputDependencies, Action<NewNode>? action = null)
 		{
-			NewNode node = new NewNode(name, inputDependencies: inputDependencies?.ToList(), orderDependencies: inputDependencies?.ToList());
+			return AddNode(group, name, inputDependencies, null, action);
+		}
+
+		static NewNode AddNode(NewGroup group, string name, string[]? inputDependencies, string[]? optionalInputDependencies, Action<NewNode>? action = null)
+		{
+			HashSet<string>? orderDependencies = [];
+			if (inputDependencies != null)
+			{
+				orderDependencies.UnionWith(inputDependencies);
+			}
+			if (optionalInputDependencies != null)
+			{
+				orderDependencies.UnionWith(optionalInputDependencies);
+			}
+
+			NewNode node = new(name, inputDependencies: inputDependencies?.ToList(), optionalInputDependencies: optionalInputDependencies?.ToList(), orderDependencies: orderDependencies?.ToList());
 			action?.Invoke(node);
 			group.Nodes.Add(node);
 			return node;
@@ -169,6 +184,553 @@ namespace HordeServer.Tests.Jobs
 			Assert.AreEqual(JobStepState.Skipped, job.Batches[3].Steps[2].State); // Post-Publish Client
 		}
 
+		[TestMethod]
+		public async Task TestStatesSkippedWhenFailureOfRequiredAsync()
+		{
+			Mock<ITemplate> templateMock = new(MockBehavior.Strict);
+			templateMock.SetupGet(x => x.InitialAgentType).Returns((string?)null);
+
+			IGraph baseGraph = await GraphCollection.AddAsync(templateMock.Object);
+
+			CreateJobOptions options = new();
+			options.Arguments.Add("-Target=Node C");
+			options.Arguments.Add("-Target=Node D");
+
+			IJob job = await JobCollection.AddAsync(JobIdUtils.GenerateNewId(), new StreamId("ue4-main"), new TemplateId("test-build"), ContentHash.SHA1("hello"), baseGraph, "Test job", CommitIdWithOrder.FromPerforceChange(123), CommitIdWithOrder.FromPerforceChange(123), options);
+
+			job = await RunBatchAsync(job, 0);
+			job = await RunStepAsync(job, 0, 0, JobStepOutcome.Success); // Setup Build
+
+			List<NewGroup> newGroups = [];
+
+			NewGroup groupA = AddGroup(newGroups);
+			AddNode(groupA, "Node A", null);
+
+			NewGroup groupB = AddGroup(newGroups);
+			AddNode(groupB, "Node B", null, ["Node A"]);
+
+			NewGroup groupC = AddGroup(newGroups);
+			AddNode(groupC, "Node C", ["Node B"]);
+
+			NewGroup groupD = AddGroup(newGroups);
+			AddNode(groupD, "Node D", ["Node B"]);
+
+			IGraph graph = await GraphCollection.AppendAsync(baseGraph, newGroups, null, null);
+			job = Deref(await job.TryUpdateGraphAsync(graph));
+
+			job = await RunBatchAsync(job, 1);
+			job = await RunStepAsync(job, 1, 0, JobStepOutcome.Success); // A
+
+			job = await RunBatchAsync(job, 2);
+			job = await RunStepAsync(job, 2, 0, JobStepOutcome.Failure); // B
+
+			Assert.AreEqual(JobStepState.Completed, job.Batches[1].Steps[0].State);
+			Assert.AreEqual(JobStepState.Completed, job.Batches[2].Steps[0].State);
+			Assert.AreEqual(JobStepState.Skipped, job.Batches[3].Steps[0].State);
+			Assert.AreEqual(JobStepState.Skipped, job.Batches[4].Steps[0].State);
+		}
+
+		[TestMethod]
+		public async Task TestStatesCompleteWhenNoFailureAsync()
+		{
+			Mock<ITemplate> templateMock = new(MockBehavior.Strict);
+			templateMock.SetupGet(x => x.InitialAgentType).Returns((string?)null);
+
+			IGraph baseGraph = await GraphCollection.AddAsync(templateMock.Object);
+
+			CreateJobOptions options = new();
+			options.Arguments.Add("-Target=Node C");
+			options.Arguments.Add("-Target=Node D");
+
+			IJob job = await JobCollection.AddAsync(JobIdUtils.GenerateNewId(), new StreamId("ue4-main"), new TemplateId("test-build"), ContentHash.SHA1("hello"), baseGraph, "Test job", CommitIdWithOrder.FromPerforceChange(123), CommitIdWithOrder.FromPerforceChange(123), options);
+
+			job = await RunBatchAsync(job, 0);
+			job = await RunStepAsync(job, 0, 0, JobStepOutcome.Success); // Setup Build
+
+			List<NewGroup> newGroups = [];
+
+			NewGroup groupA = AddGroup(newGroups);
+			AddNode(groupA, "Node A", null);
+
+			NewGroup groupB = AddGroup(newGroups);
+			AddNode(groupB, "Node B", null, ["Node A"]);
+
+			NewGroup groupC = AddGroup(newGroups);
+			AddNode(groupC, "Node C", ["Node B"]);
+
+			NewGroup groupD = AddGroup(newGroups);
+			AddNode(groupD, "Node D", ["Node B"]);
+
+			IGraph graph = await GraphCollection.AppendAsync(baseGraph, newGroups, null, null);
+			job = Deref(await job.TryUpdateGraphAsync(graph));
+
+			job = await RunBatchAsync(job, 1);
+			job = await RunStepAsync(job, 1, 0, JobStepOutcome.Success); // A
+
+			job = await RunBatchAsync(job, 2);
+			job = await RunStepAsync(job, 2, 0, JobStepOutcome.Success); // B
+
+			job = await RunBatchAsync(job, 3);
+			job = await RunStepAsync(job, 3, 0, JobStepOutcome.Success); // C
+
+			job = await RunBatchAsync(job, 4);
+			job = await RunStepAsync(job, 4, 0, JobStepOutcome.Success); // D
+
+			Assert.AreEqual(JobStepState.Completed, job.Batches[1].Steps[0].State);
+			Assert.AreEqual(JobStepState.Completed, job.Batches[2].Steps[0].State);
+			Assert.AreEqual(JobStepState.Completed, job.Batches[3].Steps[0].State);
+			Assert.AreEqual(JobStepState.Completed, job.Batches[4].Steps[0].State);
+		}
+
+		[TestMethod]
+		public async Task TestStatesCompleteWhenFailureOfOptionalAsync()
+		{
+			Mock<ITemplate> templateMock = new(MockBehavior.Strict);
+			templateMock.SetupGet(x => x.InitialAgentType).Returns((string?)null);
+
+			IGraph baseGraph = await GraphCollection.AddAsync(templateMock.Object);
+
+			CreateJobOptions options = new();
+			options.Arguments.Add("-Target=Node C");
+			options.Arguments.Add("-Target=Node D");
+
+			IJob job = await JobCollection.AddAsync(JobIdUtils.GenerateNewId(), new StreamId("ue4-main"), new TemplateId("test-build"), ContentHash.SHA1("hello"), baseGraph, "Test job", CommitIdWithOrder.FromPerforceChange(123), CommitIdWithOrder.FromPerforceChange(123), options);
+
+			job = await RunBatchAsync(job, 0);
+			job = await RunStepAsync(job, 0, 0, JobStepOutcome.Success); // Setup Build
+
+			List<NewGroup> newGroups = [];
+
+			NewGroup groupA = AddGroup(newGroups);
+			AddNode(groupA, "Node A", null);
+
+			NewGroup groupB = AddGroup(newGroups);
+			AddNode(groupB, "Node B", null, ["Node A"]);
+
+			NewGroup groupC = AddGroup(newGroups);
+			AddNode(groupC, "Node C", ["Node B"]);
+
+			NewGroup groupD = AddGroup(newGroups);
+			AddNode(groupD, "Node D", ["Node B"]);
+
+			IGraph graph = await GraphCollection.AppendAsync(baseGraph, newGroups, null, null);
+			job = Deref(await job.TryUpdateGraphAsync(graph));
+
+			job = await RunBatchAsync(job, 1);
+			job = await RunStepAsync(job, 1, 0, JobStepOutcome.Failure); // A
+
+			job = await RunBatchAsync(job, 2);
+			job = await RunStepAsync(job, 2, 0, JobStepOutcome.Success); // B
+
+			job = await RunBatchAsync(job, 3);
+			job = await RunStepAsync(job, 3, 0, JobStepOutcome.Success); // C
+
+			job = await RunBatchAsync(job, 4);
+			job = await RunStepAsync(job, 4, 0, JobStepOutcome.Success); // D
+
+			Assert.AreEqual(JobStepState.Completed, job.Batches[1].Steps[0].State);
+			Assert.AreEqual(JobStepState.Completed, job.Batches[2].Steps[0].State);
+			Assert.AreEqual(JobStepState.Completed, job.Batches[3].Steps[0].State);
+			Assert.AreEqual(JobStepState.Completed, job.Batches[4].Steps[0].State);
+		}
+
+		[TestMethod]
+		public async Task TestStatesCompleteWithMixOfRequiredAndOptionalAsync()
+		{
+			Mock<ITemplate> templateMock = new(MockBehavior.Strict);
+			templateMock.SetupGet(x => x.InitialAgentType).Returns((string?)null);
+
+			IGraph baseGraph = await GraphCollection.AddAsync(templateMock.Object);
+
+			CreateJobOptions options = new();
+			options.Arguments.Add("-Target=Node C");
+
+			IJob job = await JobCollection.AddAsync(JobIdUtils.GenerateNewId(), new StreamId("ue4-main"), new TemplateId("test-build"), ContentHash.SHA1("hello"), baseGraph, "Test job", CommitIdWithOrder.FromPerforceChange(123), CommitIdWithOrder.FromPerforceChange(123), options);
+
+			job = await RunBatchAsync(job, 0);
+			job = await RunStepAsync(job, 0, 0, JobStepOutcome.Success); // Setup Build
+
+			List<NewGroup> newGroups = [];
+
+			NewGroup groupA = AddGroup(newGroups);
+			AddNode(groupA, "Node A", null);
+
+			NewGroup groupB = AddGroup(newGroups);
+			AddNode(groupB, "Node B", null);
+
+			NewGroup groupC = AddGroup(newGroups);
+			AddNode(groupC, "Node C", ["Node A"], ["Node B"]);
+
+			IGraph graph = await GraphCollection.AppendAsync(baseGraph, newGroups, null, null);
+			job = Deref(await job.TryUpdateGraphAsync(graph));
+
+			job = await RunBatchAsync(job, 1);
+			job = await RunStepAsync(job, 1, 0, JobStepOutcome.Success); // A
+
+			job = await RunBatchAsync(job, 2);
+			job = await RunStepAsync(job, 2, 0, JobStepOutcome.Failure); // B
+
+			Assert.AreEqual(JobStepState.Ready, job.Batches[3].Steps[0].State);
+
+			job = await RunBatchAsync(job, 3);
+			job = await RunStepAsync(job, 3, 0, JobStepOutcome.Success); // C
+
+			Assert.AreEqual(JobStepState.Completed, job.Batches[3].Steps[0].State);
+		}
+
+		[TestMethod]
+		public async Task TestStatesSkippedWithMixOfRequiredAndOptionalAsync()
+		{
+			Mock<ITemplate> templateMock = new(MockBehavior.Strict);
+			templateMock.SetupGet(x => x.InitialAgentType).Returns((string?)null);
+
+			IGraph baseGraph = await GraphCollection.AddAsync(templateMock.Object);
+
+			CreateJobOptions options = new();
+			options.Arguments.Add("-Target=Node C");
+
+			IJob job = await JobCollection.AddAsync(JobIdUtils.GenerateNewId(), new StreamId("ue4-main"), new TemplateId("test-build"), ContentHash.SHA1("hello"), baseGraph, "Test job", CommitIdWithOrder.FromPerforceChange(123), CommitIdWithOrder.FromPerforceChange(123), options);
+
+			job = await RunBatchAsync(job, 0);
+			job = await RunStepAsync(job, 0, 0, JobStepOutcome.Success); // Setup Build
+
+			List<NewGroup> newGroups = [];
+
+			NewGroup groupA = AddGroup(newGroups);
+			AddNode(groupA, "Node A", null);
+
+			NewGroup groupB = AddGroup(newGroups);
+			AddNode(groupB, "Node B", null);
+
+			NewGroup groupC = AddGroup(newGroups);
+			AddNode(groupC, "Node C", ["Node A"], ["Node B"]);
+
+			IGraph graph = await GraphCollection.AppendAsync(baseGraph, newGroups, null, null);
+			job = Deref(await job.TryUpdateGraphAsync(graph));
+
+			job = await RunBatchAsync(job, 1);
+			job = await RunStepAsync(job, 1, 0, JobStepOutcome.Failure); // A
+
+			job = await RunBatchAsync(job, 2);
+			job = await RunStepAsync(job, 2, 0, JobStepOutcome.Success); // B
+
+			Assert.AreEqual(JobStepState.Skipped, job.Batches[3].Steps[0].State);
+		}
+
+		[TestMethod]
+		public async Task TestStatesReadyWhenAgentAbortedOptionalAsync()
+		{
+			Mock<ITemplate> templateMock = new(MockBehavior.Strict);
+			templateMock.SetupGet(x => x.InitialAgentType).Returns((string?)null);
+
+			IGraph baseGraph = await GraphCollection.AddAsync(templateMock.Object);
+
+			CreateJobOptions options = new();
+			options.Arguments.Add("-Target=Node B");
+
+			IJob job = await JobCollection.AddAsync(JobIdUtils.GenerateNewId(), new StreamId("ue4-main"), new TemplateId("test-build"), ContentHash.SHA1("hello"), baseGraph, "Test job", CommitIdWithOrder.FromPerforceChange(123), CommitIdWithOrder.FromPerforceChange(123), options);
+
+			job = await RunBatchAsync(job, 0);
+			job = await RunStepAsync(job, 0, 0, JobStepOutcome.Success); // Setup Build
+
+			List<NewGroup> newGroups = [];
+
+			NewGroup groupA = AddGroup(newGroups);
+			AddNode(groupA, "Node A", null);
+
+			NewGroup groupB = AddGroup(newGroups);
+			AddNode(groupB, "Node B", null, ["Node A"]);
+
+			IGraph graph = await GraphCollection.AppendAsync(baseGraph, newGroups, null, null);
+			job = Deref(await job.TryUpdateGraphAsync(graph));
+
+			// Node A is started and cancelled by agent
+			job = await RunBatchAsync(job, 1);
+			job = Deref(await job.TryUpdateStepAsync(job.Batches[1].Id, job.Batches[1].Steps[0].Id, JobStepState.Running, JobStepOutcome.Success));
+			job = Deref(await job.TryUpdateStepAsync(job.Batches[1].Id, job.Batches[1].Steps[0].Id, JobStepState.Aborted, JobStepOutcome.Failure));
+
+			Assert.AreEqual(JobStepState.Aborted, job.Batches[1].Steps[0].State);
+			Assert.IsFalse(job.Batches[1].Steps[0].AbortRequested);
+
+			// Node B should still be ready since A is optional and not aborted by user
+			Assert.AreEqual(JobStepState.Ready, job.Batches[2].Steps[0].State);
+		}
+
+		[TestMethod]
+		public async Task TestStatesSkippedWhenUserAbortedOptionalAsync()
+		{
+			Mock<ITemplate> templateMock = new(MockBehavior.Strict);
+			templateMock.SetupGet(x => x.InitialAgentType).Returns((string?)null);
+
+			IGraph baseGraph = await GraphCollection.AddAsync(templateMock.Object);
+
+			CreateJobOptions options = new();
+			options.Arguments.Add("-Target=Node B");
+
+			IJob job = await JobCollection.AddAsync(JobIdUtils.GenerateNewId(), new StreamId("ue4-main"), new TemplateId("test-build"), ContentHash.SHA1("hello"), baseGraph, "Test job", CommitIdWithOrder.FromPerforceChange(123), CommitIdWithOrder.FromPerforceChange(123), options);
+
+			job = await RunBatchAsync(job, 0);
+			job = await RunStepAsync(job, 0, 0, JobStepOutcome.Success); // Setup Build
+
+			List<NewGroup> newGroups = [];
+
+			NewGroup groupA = AddGroup(newGroups);
+			AddNode(groupA, "Node A", null);
+
+			NewGroup groupB = AddGroup(newGroups);
+			AddNode(groupB, "Node B", null, ["Node A"]);
+
+			IGraph graph = await GraphCollection.AppendAsync(baseGraph, newGroups, null, null);
+			job = Deref(await job.TryUpdateGraphAsync(graph));
+
+			// Node A is started and cancelled by user request
+			job = await RunBatchAsync(job, 1);
+			job = Deref(await job.TryUpdateStepAsync(job.Batches[1].Id, job.Batches[1].Steps[0].Id, JobStepState.Running, JobStepOutcome.Success));
+			job = Deref(await job.TryUpdateStepAsync(job.Batches[1].Id, job.Batches[1].Steps[0].Id, newAbortRequested: true));
+
+			Assert.AreEqual(JobStepState.Running, job.Batches[1].Steps[0].State);
+			Assert.IsTrue(job.Batches[1].Steps[0].AbortRequested);
+
+			// Node B should be skipped since user explicitly aborted A even though A is optional
+			Assert.AreEqual(JobStepState.Skipped, job.Batches[2].Steps[0].State);
+		}
+
+		[TestMethod]
+		public async Task TestStatesCompleteWhenFailureOfChainedOptionalAsync()
+		{
+			Mock<ITemplate> templateMock = new(MockBehavior.Strict);
+			templateMock.SetupGet(x => x.InitialAgentType).Returns((string?)null);
+
+			IGraph baseGraph = await GraphCollection.AddAsync(templateMock.Object);
+
+			CreateJobOptions options = new();
+			options.Arguments.Add("-Target=Node C");
+
+			IJob job = await JobCollection.AddAsync(JobIdUtils.GenerateNewId(), new StreamId("ue4-main"), new TemplateId("test-build"), ContentHash.SHA1("hello"), baseGraph, "Test job", CommitIdWithOrder.FromPerforceChange(123), CommitIdWithOrder.FromPerforceChange(123), options);
+
+			job = await RunBatchAsync(job, 0);
+			job = await RunStepAsync(job, 0, 0, JobStepOutcome.Success); // Setup Build
+
+			List<NewGroup> newGroups = [];
+
+			NewGroup groupA = AddGroup(newGroups);
+			AddNode(groupA, "Node A", null);
+
+			NewGroup groupB = AddGroup(newGroups);
+			AddNode(groupB, "Node B", null, ["Node A"]);
+
+			NewGroup groupC = AddGroup(newGroups);
+			AddNode(groupC, "Node C", null, ["Node B"]);
+
+			IGraph graph = await GraphCollection.AppendAsync(baseGraph, newGroups, null, null);
+			job = Deref(await job.TryUpdateGraphAsync(graph));
+
+			job = await RunBatchAsync(job, 1);
+			job = await RunStepAsync(job, 1, 0, JobStepOutcome.Failure); // A
+
+			Assert.AreEqual(JobStepState.Ready, job.Batches[2].Steps[0].State);
+
+			job = await RunBatchAsync(job, 2);
+			job = await RunStepAsync(job, 2, 0, JobStepOutcome.Failure); // B
+
+			Assert.AreEqual(JobStepState.Ready, job.Batches[3].Steps[0].State);
+
+			job = await RunBatchAsync(job, 3);
+			job = await RunStepAsync(job, 3, 0, JobStepOutcome.Success); // C
+
+			Assert.AreEqual(JobStepState.Completed, job.Batches[3].Steps[0].State);
+		}
+
+		[TestMethod]
+		public async Task TestRetryWithOptionalDependencyAsync()
+		{
+			Mock<ITemplate> templateMock = new(MockBehavior.Strict);
+			templateMock.SetupGet(x => x.InitialAgentType).Returns((string?)null);
+
+			IGraph baseGraph = await GraphCollection.AddAsync(templateMock.Object);
+
+			CreateJobOptions options = new();
+			options.Arguments.Add("-Target=Node B");
+
+			IJob job = await JobCollection.AddAsync(JobIdUtils.GenerateNewId(), new StreamId("ue4-main"), new TemplateId("test-build"), ContentHash.SHA1("hello"), baseGraph, "Test job", CommitIdWithOrder.FromPerforceChange(123), CommitIdWithOrder.FromPerforceChange(123), options);
+
+			job = await RunBatchAsync(job, 0);
+			job = await RunStepAsync(job, 0, 0, JobStepOutcome.Success); // Setup Build
+
+			List<NewGroup> newGroups = [];
+
+			NewGroup groupA = AddGroup(newGroups);
+			AddNode(groupA, "Node A", null);
+
+			NewGroup groupB = AddGroup(newGroups);
+			AddNode(groupB, "Node B", null, ["Node A"]);
+
+			IGraph graph = await GraphCollection.AppendAsync(baseGraph, newGroups, null, null);
+			job = Deref(await job.TryUpdateGraphAsync(graph));
+
+			job = await RunBatchAsync(job, 1);
+			job = await RunStepAsync(job, 1, 0, JobStepOutcome.Success); // A
+
+			job = await RunBatchAsync(job, 2);
+			job = await RunStepAsync(job, 2, 0, JobStepOutcome.Failure); // B
+
+			Assert.AreEqual(JobStepState.Completed, job.Batches[2].Steps[0].State);
+			Assert.AreEqual(JobStepOutcome.Failure, job.Batches[2].Steps[0].Outcome);
+
+			Assert.AreEqual(job.Batches.Count, 3);
+			job = Deref(await job.TryUpdateStepAsync(job.Batches[2].Id, job.Batches[2].Steps[0].Id, newRetryByUserId: UserId.Anonymous));
+			Assert.AreEqual(job.Batches.Count, 4);
+		}
+
+		[TestMethod]
+		public async Task TestRetrySkippedNodeWithAbortedOptionalDependencyAsync()
+		{
+			Mock<ITemplate> templateMock = new(MockBehavior.Strict);
+			templateMock.SetupGet(x => x.InitialAgentType).Returns((string?)null);
+
+			List<NewGroup> newGroups = [];
+
+			NewGroup groupA = AddGroup(newGroups);
+			AddNode(groupA, "Node A", null);
+
+			// Node B has ONLY optional dependency on Node A
+			NewGroup groupB = AddGroup(newGroups);
+			AddNode(groupB, "Node B", null, ["Node A"]);
+
+			IGraph graph = await GraphCollection.AppendAsync(null, newGroups, null, null);
+
+			CreateJobOptions options = new();
+			options.Arguments.Add("-Target=Node B");
+
+			IJob job = await JobCollection.AddAsync(JobIdUtils.GenerateNewId(), new StreamId("ue4-main"), new TemplateId("test-build"), ContentHash.SHA1("hello"), graph, "Test job", CommitIdWithOrder.FromPerforceChange(123), CommitIdWithOrder.FromPerforceChange(123), options);
+
+			job = await RunBatchAsync(job, 0);
+			job = Deref(await job.TryUpdateStepAsync(job.Batches[0].Id, job.Batches[0].Steps[0].Id, JobStepState.Running, JobStepOutcome.Success));
+			job = Deref(await job.TryUpdateStepAsync(job.Batches[0].Id, job.Batches[0].Steps[0].Id, newAbortRequested: true));
+			job = Deref(await job.TryUpdateStepAsync(job.Batches[0].Id, job.Batches[0].Steps[0].Id, JobStepState.Aborted, JobStepOutcome.Failure));
+
+			Assert.AreEqual(JobStepState.Aborted, job.Batches[0].Steps[0].State); // A
+			Assert.AreEqual(JobStepState.Skipped, job.Batches[1].Steps[0].State); // B
+
+			// Retry Node B
+			job = Deref(await job.TryUpdateStepAsync(job.Batches[1].Id, job.Batches[1].Steps[0].Id, newRetryByUserId: UserId.Anonymous));
+
+			// Nodes A and B to be retried
+			Assert.AreEqual(job.Batches.Count, 3);
+
+			Assert.AreEqual(JobStepState.Ready, job.Batches[1].Steps[0].State); // A
+			Assert.AreEqual(0, job.Batches[1].GroupIdx);
+			Assert.AreEqual(1, job.Batches[1].Steps.Count);
+			Assert.AreEqual(0, job.Batches[1].Steps[0].NodeIdx);
+
+			Assert.AreEqual(JobStepState.Waiting, job.Batches[2].Steps[0].State); // B
+			Assert.AreEqual(1, job.Batches[2].GroupIdx);
+			Assert.AreEqual(1, job.Batches[2].Steps.Count);
+			Assert.AreEqual(0, job.Batches[2].Steps[0].NodeIdx);
+		}
+
+		[TestMethod]
+		public async Task TestRetrySkippedNodeWithSkippedOptionalDependencyAsync()
+		{
+			Mock<ITemplate> templateMock = new(MockBehavior.Strict);
+			templateMock.SetupGet(x => x.InitialAgentType).Returns((string?)null);
+
+			List<NewGroup> newGroups = [];
+
+			NewGroup groupA = AddGroup(newGroups);
+			AddNode(groupA, "Node A", null);
+
+			NewGroup groupB = AddGroup(newGroups);
+			AddNode(groupB, "Node B", ["Node A"]);
+
+			NewGroup groupC = AddGroup(newGroups);
+			AddNode(groupC, "Node C", null, ["Node B"]);
+
+			IGraph graph = await GraphCollection.AppendAsync(null, newGroups, null, null);
+
+			CreateJobOptions options = new();
+			options.Arguments.Add("-Target=Node C");
+
+			IJob job = await JobCollection.AddAsync(JobIdUtils.GenerateNewId(), new StreamId("ue4-main"), new TemplateId("test-build"), ContentHash.SHA1("hello"), graph, "Test job", CommitIdWithOrder.FromPerforceChange(123), CommitIdWithOrder.FromPerforceChange(123), options);
+
+			job = await RunBatchAsync(job, 0);
+			job = await RunStepAsync(job, 0, 0, JobStepOutcome.Failure); // A
+
+			Assert.AreEqual(job.Batches.Count, 3);
+			Assert.AreEqual(JobStepState.Completed, job.Batches[0].Steps[0].State); // A
+			Assert.AreEqual(JobStepState.Skipped, job.Batches[1].Steps[0].State); // B
+			Assert.AreEqual(JobStepState.Skipped, job.Batches[2].Steps[0].State); // C
+
+			// Retry Node C
+			job = Deref(await job.TryUpdateStepAsync(job.Batches[2].Id, job.Batches[2].Steps[0].Id, newRetryByUserId: UserId.Anonymous));
+
+			Assert.AreEqual(job.Batches.Count, 4);
+
+			// Nodes A, B and C to be retried
+			Assert.AreEqual(JobStepState.Ready, job.Batches[1].Steps[0].State); // A
+			Assert.AreEqual(0, job.Batches[1].GroupIdx);
+			Assert.AreEqual(1, job.Batches[1].Steps.Count);
+			Assert.AreEqual(0, job.Batches[1].Steps[0].NodeIdx);
+
+			Assert.AreEqual(JobStepState.Waiting, job.Batches[2].Steps[0].State); // B
+			Assert.AreEqual(1, job.Batches[2].GroupIdx);
+			Assert.AreEqual(1, job.Batches[2].Steps.Count);
+			Assert.AreEqual(0, job.Batches[2].Steps[0].NodeIdx);
+
+			Assert.AreEqual(JobStepState.Waiting, job.Batches[3].Steps[0].State); // C
+			Assert.AreEqual(2, job.Batches[3].GroupIdx);
+			Assert.AreEqual(1, job.Batches[3].Steps.Count);
+			Assert.AreEqual(0, job.Batches[3].Steps[0].NodeIdx);
+		}
+
+		[TestMethod]
+		public async Task TestRetryCompletedNodeDoesNotRetryOptionalDependencyAsync()
+		{
+			Mock<ITemplate> templateMock = new(MockBehavior.Strict);
+			templateMock.SetupGet(x => x.InitialAgentType).Returns((string?)null);
+
+			List<NewGroup> newGroups = [];
+
+			NewGroup groupA = AddGroup(newGroups);
+			AddNode(groupA, "Node A", null);
+
+			// Node B has ONLY optional dependency on Node A
+			NewGroup groupB = AddGroup(newGroups);
+			AddNode(groupB, "Node B", null, ["Node A"]);
+
+			IGraph graph = await GraphCollection.AppendAsync(null, newGroups, null, null);
+
+			CreateJobOptions options = new();
+			options.Arguments.Add("-Target=Node B");
+
+			IJob job = await JobCollection.AddAsync(JobIdUtils.GenerateNewId(), new StreamId("ue4-main"), new TemplateId("test-build"), ContentHash.SHA1("hello"), graph, "Test job", CommitIdWithOrder.FromPerforceChange(123), CommitIdWithOrder.FromPerforceChange(123), options);
+
+			// Node A
+			job = await RunBatchAsync(job, 0);
+			job = Deref(await job.TryUpdateStepAsync(job.Batches[0].Id, job.Batches[0].Steps[0].Id, JobStepState.Running, JobStepOutcome.Success));
+			job = Deref(await job.TryUpdateStepAsync(job.Batches[0].Id, job.Batches[0].Steps[0].Id, JobStepState.Completed, JobStepOutcome.Failure));
+
+			Assert.AreEqual(JobStepState.Completed, job.Batches[0].Steps[0].State);
+			Assert.AreEqual(JobStepOutcome.Failure, job.Batches[0].Steps[0].Outcome);
+			Assert.IsFalse(job.Batches[0].Steps[0].AbortRequested);
+			Assert.AreEqual(JobStepState.Ready, job.Batches[1].Steps[0].State);
+
+			// Node B
+			job = await RunBatchAsync(job, 1);
+			job = await RunStepAsync(job, 1, 0, JobStepOutcome.Failure);
+
+			Assert.AreEqual(JobStepState.Completed, job.Batches[1].Steps[0].State);
+			Assert.AreEqual(JobStepOutcome.Failure, job.Batches[1].Steps[0].Outcome);
+
+			// Retry Node B
+			job = Deref(await job.TryUpdateStepAsync(job.Batches[1].Id, job.Batches[1].Steps[0].Id, newRetryByUserId: UserId.Anonymous));
+			Assert.AreEqual(job.Batches.Count, 3);
+			Assert.AreEqual(JobStepState.Ready, job.Batches[2].Steps[0].State);
+			Assert.AreEqual(1, job.Batches[2].GroupIdx);
+			Assert.AreEqual(1, job.Batches[2].Steps.Count);
+			Assert.AreEqual(0, job.Batches[2].Steps[0].NodeIdx);
+		}
+
 		[TestMethod]
 		public async Task TryUpdateGraphAsync()
 		{
diff --git a/Engine/Source/Programs/Horde/Plugins/Build/HordeServer.Build/Jobs/Graphs/GraphCollection.cs b/Engine/Source/Programs/Horde/Plugins/Build/HordeServer.Build/Jobs/Graphs/GraphCollection.cs
--- a/Engine/Source/Programs/Horde/Plugins/Build/HordeServer.Build/Jobs/Graphs/GraphCollection.cs
+++ b/Engine/Source/Programs/Horde/Plugins/Build/HordeServer.Build/Jobs/Graphs/GraphCollection.cs
@@ -32,12 +32,18 @@ namespace HordeServer.Jobs.Graphs
 			[BsonIgnoreIfNull]
 			public NodeOutputRef[]? Inputs { get; set; }
 
+			[BsonIgnoreIfNull]
+			public NodeOutputRef[]? OptionalInputs { get; set; }
+
 			[BsonIgnoreIfNull]
 			public string[]? OutputNames { get; set; }
 
 			[BsonIgnoreIfNull]
 			public NodeRef[] InputDependencies { get; set; }
 
+			[BsonIgnoreIfNull]
+			public NodeRef[]? OptionalInputDependencies { get; set; }
+
 			[BsonIgnoreIfNull]
 			public NodeRef[] OrderDependencies { get; set; }
 
@@ -56,10 +62,12 @@ namespace HordeServer.Jobs.Graphs
 			[BsonDictionaryOptions(DictionaryRepresentation.ArrayOfDocuments)]
 			public NodeAnnotations? Annotations { get; set; }
 
-			IReadOnlyList<NodeOutputRef> INode.Inputs => Inputs ?? Array.Empty<NodeOutputRef>();
-			IReadOnlyList<string> INode.OutputNames => OutputNames ?? Array.Empty<string>();
+			IReadOnlyList<NodeOutputRef> INode.Inputs => Inputs ?? [];
+			IReadOnlyList<NodeOutputRef> INode.OptionalInputs => OptionalInputs ?? [];
+			IReadOnlyList<string> INode.OutputNames => OutputNames ?? [];
 			IReadOnlyDictionary<string, string>? INode.Credentials => Credentials;
 			IReadOnlyDictionary<string, string>? INode.Properties => Properties;
+			NodeRef[] INode.OptionalInputDependencies => OptionalInputDependencies ?? [];
 			IReadOnlyNodeAnnotations INode.Annotations => Annotations ?? NodeAnnotations.Empty;
 
 			[BsonConstructor]
@@ -70,12 +78,14 @@ namespace HordeServer.Jobs.Graphs
 				OrderDependencies = null!;
 			}
 
-			public Node(string name, NodeOutputRef[]? inputs, string[]? outputNames, NodeRef[] inputDependencies, NodeRef[] orderDependencies, Priority priority, bool allowRetry, bool runEarly, bool warnings, Dictionary<string, string>? credentials, Dictionary<string, string>? properties, IReadOnlyNodeAnnotations? annotations)
+			public Node(string name, NodeOutputRef[]? inputs, NodeOutputRef[]? optionalInputs, string[]? outputNames, NodeRef[] inputDependencies, NodeRef[]? optionalInputDependencies, NodeRef[] orderDependencies, Priority priority, bool allowRetry, bool runEarly, bool warnings, Dictionary<string, string>? credentials, Dictionary<string, string>? properties, IReadOnlyNodeAnnotations? annotations)
 			{
 				Name = name;
 				Inputs = (inputs != null && inputs.Length > 0) ? inputs : null;
+				OptionalInputs = (optionalInputs != null && optionalInputs.Length > 0) ? optionalInputs : null;
 				OutputNames = (outputNames != null && outputNames.Length > 0) ? outputNames : null;
 				InputDependencies = inputDependencies;
+				OptionalInputDependencies = (optionalInputDependencies != null && optionalInputDependencies.Length > 0) ? optionalInputDependencies : null;
 				OrderDependencies = orderDependencies;
 				Priority = priority;
 				AllowRetry = allowRetry;
@@ -286,10 +296,17 @@ namespace HordeServer.Jobs.Graphs
 								inputs = newNodeRequest.Inputs.Select(x => nodeOutputNameToRef[x]).ToArray();
 							}
 
-							NodeRef[] inputDependencies = (newNodeRequest.InputDependencies == null) ? Array.Empty<NodeRef>() : newNodeRequest.InputDependencies.Select(x => nodeNameToRef[x]).ToArray();
-							NodeRef[] orderDependencies = (newNodeRequest.OrderDependencies == null) ? Array.Empty<NodeRef>() : newNodeRequest.OrderDependencies.Select(x => nodeNameToRef[x]).ToArray();
-							orderDependencies = orderDependencies.Union(inputDependencies).ToArray();
-							nodes.Add(new Node(newNodeRequest.Name, inputs, newNodeRequest.Outputs?.ToArray(), inputDependencies, orderDependencies, priority, allowRetry, runEarly, warnings, newNodeRequest.Credentials, newNodeRequest.Properties, newNodeRequest.Annotations));
+							NodeOutputRef[]? optionalInputs = null;
+							if (newNodeRequest.OptionalInputs != null && newNodeRequest.OptionalInputs.Count > 0)
+							{
+								optionalInputs = newNodeRequest.OptionalInputs.Select(x => nodeOutputNameToRef[x]).ToArray();
+							}
+
+							NodeRef[] inputDependencies = (newNodeRequest.InputDependencies == null) ? [] : newNodeRequest.InputDependencies.Select(x => nodeNameToRef[x]).ToArray();
+							NodeRef[] optionalInputDependencies = (newNodeRequest.OptionalInputDependencies == null) ? [] : newNodeRequest.OptionalInputDependencies.Select(x => nodeNameToRef[x]).ToArray();
+							NodeRef[] orderDependencies = (newNodeRequest.OrderDependencies == null) ? [] : newNodeRequest.OrderDependencies.Select(x => nodeNameToRef[x]).ToArray();
+							orderDependencies = orderDependencies.Union(inputDependencies).Union(optionalInputDependencies).ToArray();
+							nodes.Add(new Node(newNodeRequest.Name, inputs, optionalInputs, newNodeRequest.Outputs?.ToArray(), inputDependencies, optionalInputDependencies, orderDependencies, priority, allowRetry, runEarly, warnings, newNodeRequest.Credentials, newNodeRequest.Properties, newNodeRequest.Annotations));
 
 							NodeRef nodeRef = new NodeRef(newGroups.Count, nodeIdx);
 							nodeNameToRef.Add(newNodeRequest.Name, nodeRef);
@@ -465,7 +482,7 @@ namespace HordeServer.Jobs.Graphs
 		/// <inheritdoc/>
 		public async Task<IGraph> AddAsync(ITemplate template, CancellationToken cancellationToken)
 		{
-			Node node = new Node(IJob.SetupNodeName, null, null, Array.Empty<NodeRef>(), Array.Empty<NodeRef>(), Priority.High, true, false, true, null, null, null);
+			Node node = new Node(IJob.SetupNodeName, null, null, null, [], [], [], Priority.High, true, false, true, null, null, null);
 			NodeGroup group = new NodeGroup(template.InitialAgentType ?? "Win64", new List<Node> { node });
 
 			GraphDocument graph = new GraphDocument(new List<NodeGroup> { group }, new List<Aggregate>(), new List<Label>(), new List<GraphArtifact>());
diff --git a/Engine/Source/Programs/Horde/Plugins/Build/HordeServer.Build/Jobs/JobCollection.cs b/Engine/Source/Programs/Horde/Plugins/Build/HordeServer.Build/Jobs/JobCollection.cs
--- a/Engine/Source/Programs/Horde/Plugins/Build/HordeServer.Build/Jobs/JobCollection.cs
+++ b/Engine/Source/Programs/Horde/Plugins/Build/HordeServer.Build/Jobs/JobCollection.cs
@@ -321,11 +321,17 @@ namespace HordeServer.Jobs
 			List<JobStepOutputRef>? _inputs;
 			IReadOnlyList<JobStepOutputRef> IJobStep.Inputs => _inputs ??= CreateInputDependencyList();
 
+			List<JobStepOutputRef>? _optionalInputs;
+			IReadOnlyList<JobStepOutputRef> IJobStep.OptionalInputs => _optionalInputs ??= CreateOptionalInputDependencyList();
+
 			IReadOnlyList<string> IJobStep.OutputNames => Node.OutputNames;
 
 			List<JobStepId>? _inputDependencies;
 			IReadOnlyList<JobStepId> IJobStep.InputDependencies => _inputDependencies ??= CreateInputDependenciesList();
 
+			List<JobStepId>? _optionalInputDependencies;
+			IReadOnlyList<JobStepId> IJobStep.OptionalInputDependencies => _optionalInputDependencies ??= CreateOptionalInputDependenciesList();
+
 			List<JobStepId>? _orderDependencies;
 			IReadOnlyList<JobStepId> IJobStep.OrderDependencies => _orderDependencies ??= CreateOrderDependenciesList();
 
@@ -396,6 +402,40 @@ namespace HordeServer.Jobs
 				return dependencies;
 			}
 
+			List<JobStepOutputRef> CreateOptionalInputDependencyList()
+			{
+				List<JobStepOutputRef> dependencies = [];
+				foreach (NodeOutputRef optionalInput in Node.OptionalInputs)
+				{
+					if (Job.NodeRefToStepId.TryGetValue(optionalInput.NodeRef, out JobStepId jobStepId))
+					{
+						dependencies.Add(new JobStepOutputRef(jobStepId, optionalInput.OutputIdx));
+					}
+					else
+					{
+						MissingNodeRef(optionalInput.NodeRef);
+					}
+				}
+				return dependencies;
+			}
+
+			List<JobStepId> CreateOptionalInputDependenciesList()
+			{
+				List<JobStepId> dependencies = [];
+				foreach (NodeRef nodeRef in Node.OptionalInputDependencies)
+				{
+					if (Job.NodeRefToStepId.TryGetValue(nodeRef, out JobStepId jobStepId))
+					{
+						dependencies.Add(jobStepId);
+					}
+					else
+					{
+						MissingNodeRef(nodeRef);
+					}
+				}
+				return dependencies;
+			}
+
 			List<JobStepId> CreateOrderDependenciesList()
 			{
 				List<JobStepId> dependencies = new List<JobStepId>();
@@ -1526,6 +1566,8 @@ namespace HordeServer.Jobs
 							// Add all the dependencies to be retried
 							NodeRef[] dependencies = graph.GetNode(nodeRef).InputDependencies;
 							retryNodes.UnionWith(dependencies);
+							NodeRef[] optionalDependencies = graph.GetNode(nodeRef).OptionalInputDependencies;
+							retryNodes.UnionWith(optionalDependencies);
 						}
 						else
 						{
@@ -1919,6 +1961,13 @@ namespace HordeServer.Jobs
 				return false;
 			}
 
+			IEnumerable<string> oldOptionalInputDependencies = oldNode.OptionalInputDependencies.Select(x => oldGraph.GetNode(x).Name);
+			IEnumerable<string> newOptionalInputDependencies = newNode.OptionalInputDependencies.Select(x => newGraph.GetNode(x).Name);
+			if (!CompareListsIgnoreOrder(oldOptionalInputDependencies, newOptionalInputDependencies))
+			{
+				return false;
+			}
+
 			IEnumerable<string> oldInputs = oldNode.Inputs.Select(x => oldGraph.GetNode(x.NodeRef).OutputNames[x.OutputIdx]);
 			IEnumerable<string> newInputs = newNode.Inputs.Select(x => newGraph.GetNode(x.NodeRef).OutputNames[x.OutputIdx]);
 			if (!CompareListsIgnoreOrder(oldInputs, newInputs))
@@ -1926,6 +1975,13 @@ namespace HordeServer.Jobs
 				return false;
 			}
 
+			IEnumerable<string> oldOptionalInputs = oldNode.OptionalInputs.Select(x => oldGraph.GetNode(x.NodeRef).OutputNames[x.OutputIdx]);
+			IEnumerable<string> newOptionalInputs = newNode.OptionalInputs.Select(x => newGraph.GetNode(x.NodeRef).OutputNames[x.OutputIdx]);
+			if (!CompareListsIgnoreOrder(oldOptionalInputs, newOptionalInputs))
+			{
+				return false;
+			}
+
 			return CompareListsIgnoreOrder(oldNode.OutputNames, newNode.OutputNames);
 		}
 
@@ -2285,6 +2341,8 @@ namespace HordeServer.Jobs
 						{
 							failedNodes.Remove(node);
 						}
+
+						// Optional input dependencies are not checked here, failing optional dependencies do not cause the node to be skipped
 					}
 					else if (step.Outcome == JobStepOutcome.Failure)
 					{
@@ -2354,6 +2412,10 @@ namespace HordeServer.Jobs
 						{
 							newNodesToExecute.Add(graph.GetNode(dependency));
 						}
+						foreach (NodeRef dependency in node.OptionalInputDependencies)
+						{
+							newNodesToExecute.Add(graph.GetNode(dependency));
+						}
 					}
 				}
 			}
@@ -2411,6 +2473,13 @@ namespace HordeServer.Jobs
 								newNodesToExecute.Add(group.Nodes[dependency.NodeIdx]);
 							}
 						}
+						foreach (NodeRef dependency in node.OptionalInputDependencies)
+						{
+							if (dependency.GroupIdx == groupIdx)
+							{
+								newNodesToExecute.Add(group.Nodes[dependency.NodeIdx]);
+							}
+						}
 					}
 				}
 			}
@@ -2642,7 +2711,13 @@ namespace HordeServer.Jobs
 						if (newState == JobStepState.Waiting)
 						{
 							List<JobStepDocument> steps = GetDependentSteps(graph, node, stepForNode);
-							if (steps.Any(x => x.AbortRequested || x.State == JobStepState.Skipped || x.Outcome == JobStepOutcome.Failure))
+							HashSet<JobStepDocument> requiredSteps = new(steps);
+							List<JobStepDocument> optionalSteps = GetOptionalInputDependentSteps(graph, node, stepForNode);
+							requiredSteps.ExceptWith(optionalSteps);
+							// An optional input step aborted by the driver (not user requested) still allows the step to move to ready.
+							// A skipped optional input step does not allow continuation; only allow continuation when the dependency ran to
+							// completion and failed, not when it was skipped entirely.
+							if (steps.Any(x => x.AbortRequested || x.State == JobStepState.Skipped) || requiredSteps.Any(x => x.Outcome == JobStepOutcome.Failure))
 							{
 								newState = JobStepState.Skipped;
 								newOutcome = JobStepOutcome.Failure;
@@ -2668,6 +2743,12 @@ namespace HordeServer.Jobs
 							}
 							else if (!steps.Any(x => !x.AbortRequested && JobStepExtensions.IsPendingState(x.State)))
 							{
+								if (optionalSteps.Any(x => x.Outcome == JobStepOutcome.Failure))
+								{
+									logger.LogInformation("Allowing job {JobId}, batch {BatchId}, step {StepId} to transition to ready state when optional dependencies failed ({Dependencies})", job.Id, batch.Id, step.Id, String.Join(", ", optionalSteps.Select(x => x.Id.ToString())));
+									JobAudit(job.Id, logger, jobAudit, "Allowing job {JobId}, batch {BatchId}, step {StepId} to transition to ready state when optional dependencies failed ({Dependencies})", job.Id, batch.Id, step.Id, String.Join(", ", optionalSteps.Select(x => x.Id.ToString())));
+								}
+
 								logger.LogInformation("Transitioning job {JobId}, batch {BatchId}, step {StepId} to ready state ({Dependencies})", job.Id, batch.Id, step.Id, String.Join(", ", steps.Select(x => x.Id.ToString())));
 								JobAudit(job.Id, logger, jobAudit, "Transitioning job {JobId}, batch {BatchId}, step {StepId} to ready state ({Dependencies})", job.Id, batch.Id, step.Id, String.Join(", ", steps.Select(x => x.Id.ToString())));
 								newState = JobStepState.Ready;
@@ -2754,6 +2835,19 @@ namespace HordeServer.Jobs
 			return steps;
 		}
 
+		static List<JobStepDocument> GetOptionalInputDependentSteps(IGraph graph, INode node, Dictionary<INode, JobStepDocument> stepForNode)
+		{
+			List<JobStepDocument> steps = [];
+			foreach (NodeRef optionalInputRef in node.OptionalInputDependencies)
+			{
+				if (stepForNode.TryGetValue(graph.GetNode(optionalInputRef), out JobStepDocument? step))
+				{
+					steps.Add(step);
+				}
+			}
+			return steps;
+		}
+
 		/// <summary>
 		/// Gets the new state for a batch
 		/// </summary>
diff --git a/Engine/Source/Programs/Horde/Plugins/Build/HordeServer.Build/Jobs/JobRpcService.cs b/Engine/Source/Programs/Horde/Plugins/Build/HordeServer.Build/Jobs/JobRpcService.cs
--- a/Engine/Source/Programs/Horde/Plugins/Build/HordeServer.Build/Jobs/JobRpcService.cs
+++ b/Engine/Source/Programs/Horde/Plugins/Build/HordeServer.Build/Jobs/JobRpcService.cs
@@ -387,6 +387,12 @@ namespace HordeServer.Jobs
 					response.Inputs.Add($"{inputNode.Name}/{inputNode.OutputNames[input.OutputIdx]}");
 				}
 
+				foreach (NodeOutputRef optionalInput in node.OptionalInputs)
+				{
+					INode optionalInputNode = graph.GetNode(optionalInput.NodeRef);
+					response.OptionalInputs.Add($"{optionalInputNode.Name}/{optionalInputNode.OutputNames[optionalInput.OutputIdx]}");
+				}
+
 				response.OutputNames.Add(node.OutputNames);
 
 				foreach (INodeGroup otherGroup in graph.Groups)
@@ -401,6 +407,13 @@ namespace HordeServer.Jobs
 								response.PublishOutputs.Add(outputRef.OutputIdx);
 							}
 						}
+						foreach (NodeOutputRef outputRef in otherGroup.Nodes.SelectMany(x => x.OptionalInputs))
+						{
+							if (outputRef.NodeRef.GroupIdx == batch.GroupIdx && outputRef.NodeRef.NodeIdx == step.NodeIdx && !response.PublishOutputs.Contains(outputRef.OutputIdx))
+							{
+								response.PublishOutputs.Add(outputRef.OutputIdx);
+							}
+						}
 					}
 				}
 
@@ -688,7 +701,7 @@ namespace HordeServer.Jobs
 					List<NewNode> newNodes = new List<NewNode>();
 					foreach (RpcCreateNodeRequest node in group.Nodes)
 					{
-						NewNode newNode = new NewNode(node.Name, node.Inputs.ToList(), node.Outputs.ToList(), node.InputDependencies.ToList(), node.OrderDependencies.ToList(), (Priority)node.Priority, node.AllowRetry, node.RunEarly, node.Warnings, new Dictionary<string, string>(node.Credentials), new Dictionary<string, string>(node.Properties), new NodeAnnotations(node.Annotations));
+						NewNode newNode = new NewNode(node.Name, node.Inputs.ToList(), node.OptionalInputs.ToList(), node.Outputs.ToList(), node.InputDependencies.ToList(), node.OptionalInputDependencies.ToList(), node.OrderDependencies.ToList(), (Priority)node.Priority, node.AllowRetry, node.RunEarly, node.Warnings, new Dictionary<string, string>(node.Credentials), new Dictionary<string, string>(node.Properties), new NodeAnnotations(node.Annotations));
 						newNodes.Add(newNode);
 					}
 					newGroups.Add(new NewGroup(group.AgentType, newNodes));
diff --git a/Engine/Source/Programs/Horde/Plugins/Build/HordeServer.Build/Jobs/JobsController.cs b/Engine/Source/Programs/Horde/Plugins/Build/HordeServer.Build/Jobs/JobsController.cs
--- a/Engine/Source/Programs/Horde/Plugins/Build/HordeServer.Build/Jobs/JobsController.cs
+++ b/Engine/Source/Programs/Horde/Plugins/Build/HordeServer.Build/Jobs/JobsController.cs
@@ -898,6 +898,10 @@ namespace HordeServer.Jobs
 			{
 				response.Inputs = step.Inputs.ToList();
 			}
+			if (step.OptionalInputs.Count > 0)
+			{
+				response.OptionalInputs = step.OptionalInputs.ToList();
+			}
 			if (step.OutputNames.Count > 0)
 			{
 				response.OutputNames = step.OutputNames.ToList();
@@ -906,6 +910,10 @@ namespace HordeServer.Jobs
 			{
 				response.InputDependencies = step.InputDependencies.ToList();
 			}
+			if (step.OptionalInputDependencies.Count > 0)
+			{
+				response.OptionalInputDependencies = step.OptionalInputDependencies.ToList();
+			}
 			if (step.OrderDependencies.Count > 0)
 			{
 				response.OrderDependencies = step.OrderDependencies.ToList();
@@ -1001,8 +1009,10 @@ namespace HordeServer.Jobs
 		{
 			GetNodeResponse response = new GetNodeResponse(node.Name);
 			response.Inputs.AddRange(node.Inputs.Select(x => groups[x.NodeRef.GroupIdx].Nodes[x.NodeRef.NodeIdx].OutputNames[x.OutputIdx]));
+			response.OptionalInputs.AddRange(node.OptionalInputs.Select(x => groups[x.NodeRef.GroupIdx].Nodes[x.NodeRef.NodeIdx].OutputNames[x.OutputIdx]));
 			response.Outputs.AddRange(node.OutputNames);
 			response.InputDependencies.AddRange(node.InputDependencies.Select(x => groups[x.GroupIdx].Nodes[x.NodeIdx].Name));
+			response.OptionalInputDependencies.AddRange(node.OptionalInputDependencies.Select(x => groups[x.GroupIdx].Nodes[x.NodeIdx].Name));
 			response.OrderDependencies.AddRange(node.OrderDependencies.Select(x => groups[x.GroupIdx].Nodes[x.NodeIdx].Name));
 			response.Priority = node.Priority;
 			response.AllowRetry = node.AllowRetry;
diff --git a/Engine/Source/Programs/Shared/EpicGames.Horde/Jobs/Graphs/GraphMessages.cs b/Engine/Source/Programs/Shared/EpicGames.Horde/Jobs/Graphs/GraphMessages.cs
--- a/Engine/Source/Programs/Shared/EpicGames.Horde/Jobs/Graphs/GraphMessages.cs
+++ b/Engine/Source/Programs/Shared/EpicGames.Horde/Jobs/Graphs/GraphMessages.cs
@@ -21,6 +21,11 @@ namespace EpicGames.Horde.Jobs.Graphs
 		/// </summary>
 		public List<string> Inputs { get; set; } = [];
 
+		/// <summary>
+		/// Optional inputs for this node
+		/// </summary>
+		public List<string> OptionalInputs { get; set; } = [];
+
 		/// <summary>
 		/// Output from this node
 		/// </summary>
@@ -31,6 +36,11 @@ namespace EpicGames.Horde.Jobs.Graphs
 		/// </summary>
 		public List<string> InputDependencies { get; set; } = [];
 
+		/// <summary>
+		/// Indices of nodes which may have succeeded for this node to run
+		/// </summary>
+		public List<string> OptionalInputDependencies { get; set; } = [];
+
 		/// <summary>
 		/// Indices of nodes which must have completed for this node to run
 		/// </summary>
diff --git a/Engine/Source/Programs/Shared/EpicGames.Horde/Jobs/Graphs/IGraph.cs b/Engine/Source/Programs/Shared/EpicGames.Horde/Jobs/Graphs/IGraph.cs
--- a/Engine/Source/Programs/Shared/EpicGames.Horde/Jobs/Graphs/IGraph.cs
+++ b/Engine/Source/Programs/Shared/EpicGames.Horde/Jobs/Graphs/IGraph.cs
@@ -156,7 +156,7 @@ namespace EpicGames.Horde.Jobs.Graphs
 		/// <returns>List of dependencies</returns>
 		public static IEnumerable<INode> GetDependencies(this IGraph graph, INode node)
 		{
-			return Enumerable.Concat(node.InputDependencies, node.OrderDependencies).Select(x => graph.GetNode(x));
+			return Enumerable.Concat(node.InputDependencies, node.OrderDependencies).Concat(node.OptionalInputDependencies).Select(x => graph.GetNode(x));
 		}
 	}
 
@@ -175,6 +175,11 @@ namespace EpicGames.Horde.Jobs.Graphs
 		/// </summary>
 		public IReadOnlyList<NodeOutputRef> Inputs { get; }
 
+		/// <summary>
+		/// References to optional inputs for this node
+		/// </summary>
+		public IReadOnlyList<NodeOutputRef> OptionalInputs { get; }
+
 		/// <summary>
 		/// List of output names
 		/// </summary>
@@ -185,6 +190,11 @@ namespace EpicGames.Horde.Jobs.Graphs
 		/// </summary>
 		public NodeRef[] InputDependencies { get; }
 
+		/// <summary>
+		/// Indices of nodes which may have succeeded for this node to run
+		/// </summary>
+		public NodeRef[] OptionalInputDependencies { get; }
+
 		/// <summary>
 		/// Indices of nodes which must have completed for this node to run
 		/// </summary>
@@ -492,6 +502,11 @@ namespace EpicGames.Horde.Jobs.Graphs
 		/// </summary>
 		public List<string>? Inputs { get; set; }
 
+		/// <summary>
+		/// Optional input names
+		/// </summary>
+		public List<string>? OptionalInputs { get; set; }
+
 		/// <summary>
 		/// Output names
 		/// </summary>
@@ -502,6 +517,11 @@ namespace EpicGames.Horde.Jobs.Graphs
 		/// </summary>
 		public List<string>? InputDependencies { get; set; }
 
+		/// <summary>
+		/// List of nodes which may have succeeded when this node is run
+		/// </summary>
+		public List<string>? OptionalInputDependencies { get; set; }
+
 		/// <summary>
 		/// List of nodes which must have completed for this node to run
 		/// </summary>
@@ -547,8 +567,10 @@ namespace EpicGames.Horde.Jobs.Graphs
 		/// </summary>
 		/// <param name="name">Name of the node</param>
 		/// <param name="inputs">List of inputs for the node</param>
+		/// <param name="optionalInputs">List of optional inputs for the node</param>
 		/// <param name="outputs">List of output names for the node</param>
-		/// <param name="inputDependencies">List of nodes which must have completed succesfully for this node to run</param>
+		/// <param name="inputDependencies">List of nodes which must have completed successfully for this node to run</param>
+		/// <param name="optionalInputDependencies">List of nodes which could have completed successfully for this node to run</param>
 		/// <param name="orderDependencies">List of nodes which must have completed for this node to run</param>
 		/// <param name="priority">Priority of this node</param>
 		/// <param name="allowRetry">Whether the node can be run multiple times</param>
@@ -557,12 +579,14 @@ namespace EpicGames.Horde.Jobs.Graphs
 		/// <param name="credentials">Credentials required for this node to run</param>
 		/// <param name="properties">Properties for the node</param>
 		/// <param name="annotations">User annotations for this node</param>
-		public NewNode(string name, List<string>? inputs = null, List<string>? outputs = null, List<string>? inputDependencies = null, List<string>? orderDependencies = null, Priority? priority = null, bool? allowRetry = null, bool? runEarly = null, bool? warnings = null, Dictionary<string, string>? credentials = null, Dictionary<string, string>? properties = null, IReadOnlyNodeAnnotations? annotations = null)
+		public NewNode(string name, List<string>? inputs = null, List<string>? optionalInputs = null, List<string>? outputs = null, List<string>? inputDependencies = null, List<string>? optionalInputDependencies = null, List<string>? orderDependencies = null, Priority? priority = null, bool? allowRetry = null, bool? runEarly = null, bool? warnings = null, Dictionary<string, string>? credentials = null, Dictionary<string, string>? properties = null, IReadOnlyNodeAnnotations? annotations = null)
 		{
 			Name = name;
 			Inputs = inputs;
+			OptionalInputs = optionalInputs;
 			Outputs = outputs;
 			InputDependencies = inputDependencies;
+			OptionalInputDependencies = optionalInputDependencies;
 			OrderDependencies = orderDependencies;
 			Priority = priority;
 			AllowRetry = allowRetry;
@@ -582,7 +606,7 @@ namespace EpicGames.Horde.Jobs.Graphs
 		/// <param name="graph">Existing graph containing a node</param>
 		/// <param name="node">Node to copy</param>
 		public NewNode(IGraph graph, INode node)
-			: this(node.Name, [.. node.Inputs.Select(x => graph.GetNode(x.NodeRef).Name)], [.. node.OutputNames], [.. node.InputDependencies.Select(x => graph.GetNode(x).Name)], [.. node.OrderDependencies.Select(x => graph.GetNode(x).Name)], node.Priority, node.AllowRetry, node.RunEarly, node.Warnings, node.Credentials?.ToDictionary(x => x.Key, x => x.Value), node.Properties?.ToDictionary(x => x.Key, x => x.Value), node.Annotations)
+			: this(node.Name, [.. node.Inputs.Select(x => graph.GetNode(x.NodeRef).Name)], [.. node.OptionalInputs.Select(x => graph.GetNode(x.NodeRef).Name)], [.. node.OutputNames], [.. node.InputDependencies.Select(x => graph.GetNode(x).Name)], [.. node.OptionalInputDependencies.Select(x => graph.GetNode(x).Name)], [.. node.OrderDependencies.Select(x => graph.GetNode(x).Name)], node.Priority, node.AllowRetry, node.RunEarly, node.Warnings, node.Credentials?.ToDictionary(x => x.Key, x => x.Value), node.Properties?.ToDictionary(x => x.Key, x => x.Value), node.Annotations)
 		{
 		}
 	}
diff --git a/Engine/Source/Programs/Shared/EpicGames.Horde/Jobs/IJob.cs b/Engine/Source/Programs/Shared/EpicGames.Horde/Jobs/IJob.cs
--- a/Engine/Source/Programs/Shared/EpicGames.Horde/Jobs/IJob.cs
+++ b/Engine/Source/Programs/Shared/EpicGames.Horde/Jobs/IJob.cs
@@ -1456,6 +1456,7 @@ namespace EpicGames.Horde.Jobs
 			foreach (INode node in nodes)
 			{
 				dependencies.UnionWith(node.InputDependencies.Select(x => groups[x.GroupIdx].Nodes[x.NodeIdx]));
+				dependencies.UnionWith(node.OptionalInputDependencies.Select(x => groups[x.GroupIdx].Nodes[x.NodeIdx]));
 				dependencies.UnionWith(node.OrderDependencies.Select(x => groups[x.GroupIdx].Nodes[x.NodeIdx]));
 			}
 
@@ -1505,6 +1506,11 @@ namespace EpicGames.Horde.Jobs
 		/// </summary>
 		public IReadOnlyList<JobStepOutputRef> Inputs { get; }
 
+		/// <summary>
+		/// References to optional inputs for this node
+		/// </summary>
+		public IReadOnlyList<JobStepOutputRef> OptionalInputs { get; }
+
 		/// <summary>
 		/// List of output names
 		/// </summary>
@@ -1515,6 +1521,11 @@ namespace EpicGames.Horde.Jobs
 		/// </summary>
 		public IReadOnlyList<JobStepId> InputDependencies { get; }
 
+		/// <summary>
+		/// Indices of nodes which could have succeeded for this node to run
+		/// </summary>
+		public IReadOnlyList<JobStepId> OptionalInputDependencies { get; }
+
 		/// <summary>
 		/// Indices of nodes which must have completed for this node to run
 		/// </summary>
diff --git a/Engine/Source/Programs/Shared/EpicGames.Horde/Jobs/JobMessages.cs b/Engine/Source/Programs/Shared/EpicGames.Horde/Jobs/JobMessages.cs
--- a/Engine/Source/Programs/Shared/EpicGames.Horde/Jobs/JobMessages.cs
+++ b/Engine/Source/Programs/Shared/EpicGames.Horde/Jobs/JobMessages.cs
@@ -960,6 +960,11 @@ namespace EpicGames.Horde.Jobs
 		/// </summary>
 		public List<JobStepOutputRef>? Inputs { get; set; }
 
+		/// <summary>
+		/// References to optional inputs for this node
+		/// </summary>
+		public List<JobStepOutputRef>? OptionalInputs { get; set; }
+
 		/// <summary>
 		/// List of output names
 		/// </summary>
@@ -970,6 +975,11 @@ namespace EpicGames.Horde.Jobs
 		/// </summary>
 		public List<JobStepId>? InputDependencies { get; set; }
 
+		/// <summary>
+		/// Indices of nodes which may have succeeded for this node to run
+		/// </summary>
+		public List<JobStepId>? OptionalInputDependencies { get; set; }
+
 		/// <summary>
 		/// Indices of nodes which must have completed for this node to run
 		/// </summary>
