Horde: Updating schema documentation

#rnx

diff --git a/Engine/Source/Programs/Horde/Docs/Config/Schema/AclActions.md b/Engine/Source/Programs/Horde/Docs/Config/Schema/AclActions.md
--- a/Engine/Source/Programs/Horde/Docs/Config/Schema/AclActions.md
+++ b/Engine/Source/Programs/Horde/Docs/Config/Schema/AclActions.md
@@ -37,32 +37,6 @@
 | `DeleteAccount` |  |
 | `ViewAccount` |  |
 
-## Notices
-
-| Name | Description |
-| ---- | ----------- |
-| `CreateNotice` |  |
-| `UpdateNotice` |  |
-| `DeleteNotice` |  |
-
-## ServiceAccounts
-
-| Name | Description |
-| ---- | ----------- |
-| `CreateAccount` |  |
-| `UpdateAccount` |  |
-| `DeleteAccount` |  |
-| `ViewAccount` |  |
-
-## Accounts
-
-| Name | Description |
-| ---- | ----------- |
-| `CreateAccount` |  |
-| `UpdateAccount` |  |
-| `DeleteAccount` |  |
-| `ViewAccount` |  |
-
 ## Artifacts
 
 | Name | Description |
@@ -171,6 +145,114 @@
 | `DeleteAccount` |  |
 | `ViewAccount` |  |
 
+## Artifacts
+
+| Name | Description |
+| ---- | ----------- |
+| `ReadArtifact` |  |
+| `WriteArtifact` |  |
+| `DeleteArtifact` |  |
+| `UploadArtifact` |  |
+| `DownloadArtifact` |  |
+
+## Bisect
+
+| Name | Description |
+| ---- | ----------- |
+| `CreateBisectTask` |  |
+| `UpdateBisectTask` |  |
+| `ViewBisectTask` |  |
+
+## BuildHealth
+
+| Name | Description |
+| ---- | ----------- |
+| `AddUpdateDeleteBuildHealthFilter` |  |
+
+## Devices
+
+| Name | Description |
+| ---- | ----------- |
+| `DeviceRead` |  |
+| `DeviceWrite` |  |
+
+## Jobs
+
+| Name | Description |
+| ---- | ----------- |
+| `CreateJob` |  |
+| `UpdateJob` |  |
+| `DeleteJob` |  |
+| `ExecuteJob` |  |
+| `RetryJobStep` |  |
+| `ViewJob` |  |
+
+## Notices
+
+| Name | Description |
+| ---- | ----------- |
+| `CreateNotice` |  |
+| `UpdateNotice` |  |
+| `DeleteNotice` |  |
+
+## Notifications
+
+| Name | Description |
+| ---- | ----------- |
+| `CreateSubscription` |  |
+
+## Projects
+
+| Name | Description |
+| ---- | ----------- |
+| `CreateProject` |  |
+| `DeleteProject` |  |
+| `UpdateProject` |  |
+| `ViewProject` |  |
+
+## Replicators
+
+| Name | Description |
+| ---- | ----------- |
+| `UpdateReplicator` |  |
+| `ViewReplicator` |  |
+
+## ServiceAccounts
+
+| Name | Description |
+| ---- | ----------- |
+| `CreateAccount` |  |
+| `UpdateAccount` |  |
+| `DeleteAccount` |  |
+| `ViewAccount` |  |
+
+## Streams
+
+| Name | Description |
+| ---- | ----------- |
+| `CreateStream` |  |
+| `UpdateStream` |  |
+| `DeleteStream` |  |
+| `ViewStream` |  |
+| `ViewChanges` |  |
+| `ViewTemplate` |  |
+
+## TestData
+
+| Name | Description |
+| ---- | ----------- |
+| `TestHealthRead` |  |
+| `TestHealthWrite` |  |
+
+## Accounts
+
+| Name | Description |
+| ---- | ----------- |
+| `CreateAccount` |  |
+| `UpdateAccount` |  |
+| `DeleteAccount` |  |
+| `ViewAccount` |  |
+
 ## Agents
 
 | Name | Description |
@@ -462,6 +544,12 @@
 | `UpdateReplicator` |  |
 | `ViewReplicator` |  |
 
+## Secrets
+
+| Name | Description |
+| ---- | ----------- |
+| `ViewSecret` | View a credential |
+
 ## ServiceAccounts
 
 | Name | Description |
@@ -633,6 +721,12 @@
 | `UpdateReplicator` |  |
 | `ViewReplicator` |  |
 
+## Secrets
+
+| Name | Description |
+| ---- | ----------- |
+| `ViewSecret` |  |
+
 ## ServiceAccounts
 
 | Name | Description |
@@ -808,7 +902,7 @@
 
 | Name | Description |
 | ---- | ----------- |
-| `ViewSecret` | View a credential |
+| `ViewSecret` |  |
 
 ## ServiceAccounts
 
@@ -1011,19 +1105,6 @@
 | `DownloadSoftware` |  |
 | `DeleteSoftware` |  |
 
-## Storage
-
-| Name | Description |
-| ---- | ----------- |
-| `ReadBlobs` | Ability to read blobs from the storage service |
-| `WriteBlobs` | Ability to write blobs to the storage service |
-| `ReadAliases` | Ability to read aliases from the storage service |
-| `WriteAliases` | Ability to write aliases to the storage service |
-| `DeleteAliases` | Ability to write aliases to the storage service |
-| `ReadRefs` | Ability to read refs from the storage service |
-| `WriteRefs` | Ability to write refs to the storage service |
-| `DeleteRefs` | Ability to delete refs |
-
 ## Streams
 
 | Name | Description |
@@ -1035,6 +1116,12 @@
 | `ViewChanges` |  |
 | `ViewTemplate` |  |
 
+## Symbols
+
+| Name | Description |
+| ---- | ----------- |
+| `ReadSymbols` | Ability to download symbols |
+
 ## TestData
 
 | Name | Description |
@@ -1205,14 +1292,14 @@
 
 | Name | Description |
 | ---- | ----------- |
-| `ReadBlobs` |  |
-| `WriteBlobs` |  |
-| `ReadAliases` |  |
-| `WriteAliases` |  |
-| `DeleteAliases` |  |
-| `ReadRefs` |  |
-| `WriteRefs` |  |
-| `DeleteRefs` |  |
+| `ReadBlobs` | Ability to read blobs from the storage service |
+| `WriteBlobs` | Ability to write blobs to the storage service |
+| `ReadAliases` | Ability to read aliases from the storage service |
+| `WriteAliases` | Ability to write aliases to the storage service |
+| `DeleteAliases` | Ability to write aliases to the storage service |
+| `ReadRefs` | Ability to read refs from the storage service |
+| `WriteRefs` | Ability to write refs to the storage service |
+| `DeleteRefs` | Ability to delete refs |
 
 ## Streams
 
@@ -1229,7 +1316,7 @@
 
 | Name | Description |
 | ---- | ----------- |
-| `ReadSymbols` | Ability to download symbols |
+| `ReadSymbols` |  |
 
 ## TestData
 
diff --git a/Engine/Source/Programs/Horde/Docs/Config/Schema/Globals.md b/Engine/Source/Programs/Horde/Docs/Config/Schema/Globals.md
--- a/Engine/Source/Programs/Horde/Docs/Config/Schema/Globals.md
+++ b/Engine/Source/Programs/Horde/Docs/Config/Schema/Globals.md
@@ -59,24 +59,16 @@
 
 Name | Description
 ---- | -----------
-`analytics` | [AnalyticsConfig](#analyticsconfig)<br>Configuration for the analytics plugin
 `build` | [BuildConfig](#buildconfig)<br>Configuration for the build plugin
+`analytics` | [AnalyticsConfig](#analyticsconfig)<br>Configuration for the analytics plugin
 `compute` | [ComputeConfig](#computeconfig)<br>Configuration for the compute plugin
-`experimental` | [ExperimentalConfig](#experimentalconfig)<br>Configuration for the experimental plugin
-`exampleplugin` | [ExamplePluginConfig](#examplepluginconfig)<br>Configuration for the exampleplugin plugin
 `secrets` | [SecretsConfig](#secretsconfig)<br>Configuration for the secrets plugin
-`storage` | [StorageConfig](#storageconfig)<br>Configuration for the storage plugin
+`exampleplugin` | [ExamplePluginConfig](#examplepluginconfig)<br>Configuration for the exampleplugin plugin
+`experimental` | [ExperimentalConfig](#experimentalconfig)<br>Configuration for the experimental plugin
 `symbols` | [SymbolsConfig](#symbolsconfig)<br>Configuration for the symbols plugin
+`storage` | [StorageConfig](#storageconfig)<br>Configuration for the storage plugin
 `tools` | [ToolsConfig](#toolsconfig)<br>Configuration for the tools plugin
 
-## AnalyticsConfig
-
-Config settings for analytics
-
-Name | Description
----- | -----------
-`stores` | [TelemetryStoreConfig](Telemetry.md)`[]`<br>Metrics to aggregate on the Horde server
-
 ## BuildConfig
 
 Configuration for the build plugin
@@ -251,6 +243,14 @@
 `excludeActions` | `string[]`<br>Actions to exclude from the inherited actions
 `extends` | `string[]`<br>Other profiles to extend from
 
+## AnalyticsConfig
+
+Config settings for analytics
+
+Name | Description
+---- | -----------
+`stores` | [TelemetryStoreConfig](Telemetry.md)`[]`<br>Metrics to aggregate on the Horde server
+
 ## ComputeConfig
 
 Configuration for the compute system
@@ -447,22 +447,6 @@
 `description` | `string`<br>Human-readable description
 `computeId` | `string`<br>Compute ID for this network (used when allocating compute resources)
 
-## ExperimentalConfig
-
-Configuration for the experimental plugin
-
-Name | Description
----- | -----------
-`notifications` | [NotificationConfig](Notification.md)`[]`<br>Notifications to be sent on the Horde server
-
-## ExamplePluginConfig
-
-Configuration for the ExamplePlugin plugin
-
-Name | Description
----- | -----------
-`greetingMessage` | `string`<br>The Greeting provided by the example plugin
-
 ## SecretsConfig
 
 Configuration for the secrets system
@@ -545,6 +529,41 @@
 ---- | -----------
 `vaultUri` | `string`<br>The Uri to the Azure Key Vault.
 
+## ExamplePluginConfig
+
+Configuration for the ExamplePlugin plugin
+
+Name | Description
+---- | -----------
+`greetingMessage` | `string`<br>The Greeting provided by the example plugin
+
+## ExperimentalConfig
+
+Configuration for the experimental plugin
+
+Name | Description
+---- | -----------
+`notifications` | [NotificationConfig](Notification.md)`[]`<br>Notifications to be sent on the Horde server
+
+## SymbolsConfig
+
+Configuration for the tools system
+
+Name | Description
+---- | -----------
+`stores` | [SymbolStoreConfig](#symbolstoreconfig)`[]`<br>List of symbol stores
+
+## SymbolStoreConfig
+
+Configuration for a symbol store
+
+Name | Description
+---- | -----------
+`id` | `string`<br>Identifier for this store
+`namespaceId` | `string`<br>Configuration for the symbol store backend
+`public` | `boolean`<br>Whether to make this store available without auth
+`acl` | [AclConfig](#aclconfig)<br>Access to the symbol store
+
 ## StorageConfig
 
 Configuration for storage
@@ -617,25 +636,6 @@
 `enableAliases` | `boolean`<br>Support querying exports by their aliases
 `acl` | [AclConfig](#aclconfig)<br>Access list for this namespace
 
-## SymbolsConfig
-
-Configuration for the tools system
-
-Name | Description
----- | -----------
-`stores` | [SymbolStoreConfig](#symbolstoreconfig)`[]`<br>List of symbol stores
-
-## SymbolStoreConfig
-
-Configuration for a symbol store
-
-Name | Description
----- | -----------
-`id` | `string`<br>Identifier for this store
-`namespaceId` | `string`<br>Configuration for the symbol store backend
-`public` | `boolean`<br>Whether to make this store available without auth
-`acl` | [AclConfig](#aclconfig)<br>Access to the symbol store
-
 ## ToolsConfig
 
 Configuration for the tools system
diff --git a/Engine/Source/Programs/Horde/Docs/Deployment/ServerSettings.md b/Engine/Source/Programs/Horde/Docs/Deployment/ServerSettings.md
--- a/Engine/Source/Programs/Horde/Docs/Deployment/ServerSettings.md
+++ b/Engine/Source/Programs/Horde/Docs/Deployment/ServerSettings.md
@@ -127,53 +127,16 @@
 
 Name | Description
 ---- | -----------
-`analytics` | [AnalyticsServerConfig](#analyticsserverconfig)<br>Configuration for the analytics plugin
 `build` | [BuildServerConfig](#buildserverconfig)<br>Configuration for the build plugin
+`analytics` | [AnalyticsServerConfig](#analyticsserverconfig)<br>Configuration for the analytics plugin
 `compute` | [ComputeServerConfig](#computeserverconfig)<br>Configuration for the compute plugin
-`experimental` | [ExperimentalServerConfig](#experimentalserverconfig)<br>Configuration for the experimental plugin
-`exampleplugin` | [ExamplePluginServerConfig](#examplepluginserverconfig)<br>Configuration for the exampleplugin plugin
 `secrets` | [SecretsServerConfig](#secretsserverconfig)<br>Configuration for the secrets plugin
-`storage` | [StorageServerConfig](#storageserverconfig)<br>Configuration for the storage plugin
+`exampleplugin` | [ExamplePluginServerConfig](#examplepluginserverconfig)<br>Configuration for the exampleplugin plugin
+`experimental` | [ExperimentalServerConfig](#experimentalserverconfig)<br>Configuration for the experimental plugin
 `symbols` | [PluginServerConfig](#pluginserverconfig)<br>Configuration for the symbols plugin
+`storage` | [StorageServerConfig](#storageserverconfig)<br>Configuration for the storage plugin
 `tools` | [ToolsServerConfig](#toolsserverconfig)<br>Configuration for the tools plugin
 
-## AnalyticsServerConfig
-
-Server configuration for the analytics system
-
-Name | Description
----- | -----------
-`sinks` | [TelemetrySinkConfig](#telemetrysinkconfig)<br>Settings for the various telemetry sinks
-`enabled` | `boolean`<br>Whether the plugin should be enabled or not
-
-## TelemetrySinkConfig
-
-Telemetry sinks
-
-Name | Description
----- | -----------
-`epic` | [EpicTelemetryConfig](#epictelemetryconfig)<br>Settings for the Epic telemetry sink
-`mongo` | [MongoTelemetryConfig](#mongotelemetryconfig)<br>Settings for the MongoDB telemetry sink
-
-## EpicTelemetryConfig
-
-Configuration for the telemetry sink
-
-Name | Description
----- | -----------
-`url` | `string`<br>Base URL for the telemetry server
-`appId` | `string`<br>Application name to send in the event messages
-`enabled` | `boolean`<br>Whether to enable this sink
-
-## MongoTelemetryConfig
-
-Configuration for the telemetry sink
-
-Name | Description
----- | -----------
-`retainDays` | `number`<br>Number of days worth of telmetry events to keep
-`enabled` | `boolean`<br>Whether to enable this sink
-
 ## BuildServerConfig
 
 Static configuration for the build plugin
@@ -318,6 +281,43 @@
 `maxChildCount` | `integer`<br>Maximum number of children in each node
 `sliceThreshold` | `integer`<br>Threshold hash value for splitting interior nodes
 
+## AnalyticsServerConfig
+
+Server configuration for the analytics system
+
+Name | Description
+---- | -----------
+`sinks` | [TelemetrySinkConfig](#telemetrysinkconfig)<br>Settings for the various telemetry sinks
+`enabled` | `boolean`<br>Whether the plugin should be enabled or not
+
+## TelemetrySinkConfig
+
+Telemetry sinks
+
+Name | Description
+---- | -----------
+`epic` | [EpicTelemetryConfig](#epictelemetryconfig)<br>Settings for the Epic telemetry sink
+`mongo` | [MongoTelemetryConfig](#mongotelemetryconfig)<br>Settings for the MongoDB telemetry sink
+
+## EpicTelemetryConfig
+
+Configuration for the telemetry sink
+
+Name | Description
+---- | -----------
+`url` | `string`<br>Base URL for the telemetry server
+`appId` | `string`<br>Application name to send in the event messages
+`enabled` | `boolean`<br>Whether to enable this sink
+
+## MongoTelemetryConfig
+
+Configuration for the telemetry sink
+
+Name | Description
+---- | -----------
+`retainDays` | `number`<br>Number of days worth of telmetry events to keep
+`enabled` | `boolean`<br>Whether to enable this sink
+
 ## ComputeServerConfig
 
 Static configuration for the compute plugin
@@ -363,12 +363,13 @@
 `ComputeQueueAwsMetric` | A no-op strategy that reports metrics to let an external AWS auto-scaling policy scale the fleet
 `LeaseUtilizationAwsMetric` | A no-op strategy that reports metrics to let an external AWS auto-scaling policy scale the fleet
 
-## ExperimentalServerConfig
+## SecretsServerConfig
 
-Server configuration for the experimental plugin
+Static configuration for the secrets plugin
 
 Name | Description
 ---- | -----------
+`withAws` | `boolean`<br>Whether to enable Amazon Web Services (AWS) specific features
 `enabled` | `boolean`<br>Whether the plugin should be enabled or not
 
 ## ExamplePluginServerConfig
@@ -379,13 +380,20 @@
 ---- | -----------
 `enabled` | `boolean`<br>Whether the plugin should be enabled or not
 
-## SecretsServerConfig
+## ExperimentalServerConfig
 
-Static configuration for the secrets plugin
+Server configuration for the experimental plugin
 
 Name | Description
 ---- | -----------
-`withAws` | `boolean`<br>Whether to enable Amazon Web Services (AWS) specific features
 `enabled` | `boolean`<br>Whether the plugin should be enabled or not
 
+## PluginServerConfig
+
+Base class for plugin server config objects
+
+Name | Description
+---- | -----------
+`enabled` | `boolean`<br>Whether the plugin should be enabled or not
+
 ## StorageServerConfig
@@ -447,14 +455,6 @@
 `AssumeRole` | Assume a particular role. Should specify ARN in
 `AssumeRoleWebIdentity` | Assume a particular role using the current environment variables.
 
-## PluginServerConfig
-
-Base class for plugin server config objects
-
-Name | Description
----- | -----------
-`enabled` | `boolean`<br>Whether the plugin should be enabled or not
-
 ## ToolsServerConfig
 
 Server configuration for bundled tools
