Horde: Emit telemetry on lease completion

Add a LeaseCompleteTelemetry event that fires when any lease finishes.

diff --git a/Engine/Source/Programs/Horde/Plugins/Compute/HordeServer.Compute/Agents/Leases/LeaseService.cs b/Engine/Source/Programs/Horde/Plugins/Compute/HordeServer.Compute/Agents/Leases/LeaseService.cs
new file mode 100644
--- /dev/null
+++ b/Engine/Source/Programs/Horde/Plugins/Compute/HordeServer.Compute/Agents/Leases/LeaseService.cs
@@ -0,0 +1,59 @@
+// Copyright Epic Games, Inc. All Rights Reserved.
+
+using HordeServer.Telemetry;
+using Microsoft.Extensions.Hosting;
+using Microsoft.Extensions.Logging;
+using Microsoft.Extensions.Options;
+
+namespace HordeServer.Agents.Leases
+{
+	/// <summary>
+	/// Service that emits telemetry when leases complete.
+	/// </summary>
+	public class LeaseService(
+		LeaseCollection leaseCollection,
+		ITelemetryWriter telemetryWriter,
+		IOptionsMonitor<ComputeConfig> computeConfig,
+		ILogger<LeaseService> logger)
+		: IHostedService
+	{
+		private readonly ILogger _logger = logger;
+
+		/// <inheritdoc/>
+		public Task StartAsync(CancellationToken cancellationToken)
+		{
+			leaseCollection.OnLeaseComplete += OnLeaseComplete;
+			return Task.CompletedTask;
+		}
+
+		/// <inheritdoc/>
+		public Task StopAsync(CancellationToken cancellationToken)
+		{
+			leaseCollection.OnLeaseComplete -= OnLeaseComplete;
+			return Task.CompletedTask;
+		}
+
+		private void OnLeaseComplete(ILease lease)
+		{
+			if (!telemetryWriter.Enabled)
+			{
+				return;
+			}
+
+			EpicGames.Horde.Telemetry.TelemetryStoreId telemetryStoreId = computeConfig.CurrentValue.TelemetryStoreId;
+			if (telemetryStoreId.IsEmpty)
+			{
+				return;
+			}
+
+			try
+			{
+				telemetryWriter.WriteEvent(telemetryStoreId, new LeaseCompleteTelemetry(lease));
+			}
+			catch (Exception ex)
+			{
+				_logger.LogError(ex, "Failed to write lease completion telemetry for lease {LeaseId}", lease.Id);
+			}
+		}
+	}
+}
diff --git a/Engine/Source/Programs/Horde/Plugins/Compute/HordeServer.Compute/Agents/Leases/LeaseTelemetry.cs b/Engine/Source/Programs/Horde/Plugins/Compute/HordeServer.Compute/Agents/Leases/LeaseTelemetry.cs
new file mode 100644
--- /dev/null
+++ b/Engine/Source/Programs/Horde/Plugins/Compute/HordeServer.Compute/Agents/Leases/LeaseTelemetry.cs
@@ -0,0 +1,188 @@
+// Copyright Epic Games, Inc. All Rights Reserved.
+
+using System.ComponentModel.DataAnnotations.Schema;
+using EpicGames.Analytics;
+using EpicGames.Analytics.Telemetry;
+using EpicGames.Horde.Agents;
+using EpicGames.Horde.Agents.Leases;
+using EpicGames.Horde.Agents.Pools;
+using EpicGames.Horde.Agents.Sessions;
+using EpicGames.Horde.Logs;
+using EpicGames.Horde.Streams;
+
+namespace HordeServer.Agents.Leases
+{
+	/// <summary>
+	/// Telemetry record emitted when a lease completes. Models an <see cref="ILease"/>.
+	/// </summary>
+	[AnalyticsTableGen]
+	[Table("ingest.horde.state_lease_complete")]
+	public record LeaseCompleteTelemetry : AbstractTelemetryRecord
+	{
+		/// <summary>
+		/// The lease id.
+		/// </summary>
+		[Column("lease_id")]
+		public LeaseId LeaseId { get; init; }
+
+		/// <summary>
+		/// The parent lease id, if any.
+		/// </summary>
+		[Column("parent_id")]
+		public LeaseId? ParentId { get; init; }
+
+		/// <summary>
+		/// Name of the lease.
+		/// </summary>
+		[Column("name")]
+		public string Name { get; init; }
+
+		/// <summary>
+		/// The agent that executed the lease.
+		/// </summary>
+		[Column("agent_id")]
+		public AgentId AgentId { get; init; }
+
+		/// <summary>
+		/// The session in which the lease ran.
+		/// </summary>
+		[Column("session_id")]
+		public SessionId SessionId { get; init; }
+
+		/// <summary>
+		/// The stream associated with the lease, if any.
+		/// </summary>
+		[Column("stream_id")]
+		public StreamId? StreamId { get; init; }
+
+		/// <summary>
+		/// The pool associated with the lease, if any.
+		/// </summary>
+		[Column("pool_id")]
+		public PoolId? PoolId { get; init; }
+
+		/// <summary>
+		/// The log for this lease, if any.
+		/// </summary>
+		[Column("log_id")]
+		public LogId? LogId { get; init; }
+
+		/// <summary>
+		/// Time at which the lease started.
+		/// </summary>
+		[Column("start_time_utc")]
+		public DateTime StartTimeUtc { get; init; }
+
+		/// <summary>
+		/// Time at which the lease finished.
+		/// </summary>
+		[Column("finish_time_utc")]
+		public DateTime FinishTimeUtc { get; init; }
+
+		/// <summary>
+		/// Duration of the lease in seconds.
+		/// </summary>
+		[Column("duration_secs")]
+		public double DurationSecs { get; init; }
+
+		/// <summary>
+		/// Outcome of the lease.
+		/// </summary>
+		[Column("outcome")]
+		public LeaseOutcome Outcome { get; init; }
+
+		/// <summary>
+		/// Reason for the outcome.
+		/// </summary>
+		[Column("outcome_reason")]
+		public LeaseOutcomeReason OutcomeReason { get; init; }
+
+		/// <summary>
+		/// Time spent in setup, in seconds.
+		/// </summary>
+		[Column("setup_time_secs")]
+		public double? SetupTimeSecs { get; init; }
+
+		/// <summary>
+		/// Time spent in teardown, in seconds.
+		/// </summary>
+		[Column("teardown_time_secs")]
+		public double? TeardownTimeSecs { get; init; }
+
+		/// <summary>
+		/// Number of logical CPU cores available.
+		/// </summary>
+		[Column("cpu_count")]
+		public int? CpuCount { get; init; }
+
+		/// <summary>
+		/// Machine-wide CPU utilization during the lease.
+		/// </summary>
+		[Column("global_cpu_utilization")]
+		public double? GlobalCpuUtilization { get; init; }
+
+		/// <summary>
+		/// Machine-wide user CPU time in seconds.
+		/// </summary>
+		[Column("cpu_user_time_secs")]
+		public double? CpuUserTimeSecs { get; init; }
+
+		/// <summary>
+		/// Machine-wide system CPU time in seconds.
+		/// </summary>
+		[Column("cpu_system_time_secs")]
+		public double? CpuSystemTimeSecs { get; init; }
+
+		/// <summary>
+		/// Machine-wide idle CPU time in seconds.
+		/// </summary>
+		[Column("cpu_idle_time_secs")]
+		public double? CpuIdleTimeSecs { get; init; }
+
+		/// <summary>
+		/// Default event name for this telemetry record.
+		/// </summary>
+		public const string DefaultEventName = "State.LeaseComplete";
+
+		/// <summary>
+		/// Current schema version.
+		/// </summary>
+		public static readonly int CurrentSchemaVersion = 1;
+
+		/// <summary>
+		/// Constructs a lease completion telemetry event from a completed lease.
+		/// </summary>
+		/// <param name="lease">The completed lease.</param>
+		public LeaseCompleteTelemetry(ILease lease) : base(DefaultEventName, CurrentSchemaVersion)
+		{
+			LeaseId = lease.Id;
+			ParentId = lease.ParentId;
+			Name = lease.Name;
+			AgentId = lease.AgentId;
+			SessionId = lease.SessionId;
+			StreamId = lease.StreamId;
+			PoolId = lease.PoolId;
+			LogId = lease.LogId;
+			StartTimeUtc = lease.StartTime;
+			FinishTimeUtc = lease.FinishTime ?? lease.StartTime;
+			DurationSecs = (FinishTimeUtc - StartTimeUtc).TotalSeconds;
+			Outcome = lease.Outcome;
+			OutcomeReason = lease.OutcomeReason;
+			SetupTimeSecs = lease.Metrics?.SetupTime.TotalSeconds;
+			TeardownTimeSecs = lease.Metrics?.TeardownTime.TotalSeconds;
+			CpuCount = lease.Metrics?.CpuCount;
+			GlobalCpuUtilization = lease.Metrics?.GlobalCpuUtilization;
+			CpuUserTimeSecs = lease.Metrics?.GlobalCpuTime.User.TotalSeconds;
+			CpuSystemTimeSecs = lease.Metrics?.GlobalCpuTime.System.TotalSeconds;
+			CpuIdleTimeSecs = lease.Metrics?.GlobalCpuTime.Idle.TotalSeconds;
+		}
+
+		/// <summary>
+		/// Default constructor for ORM mapping.
+		/// </summary>
+		public LeaseCompleteTelemetry() : base(DefaultEventName, CurrentSchemaVersion)
+		{
+			Name = String.Empty;
+		}
+	}
+}
diff --git a/Engine/Source/Programs/Horde/Plugins/Compute/HordeServer.Compute/ComputeConfig.cs b/Engine/Source/Programs/Horde/Plugins/Compute/HordeServer.Compute/ComputeConfig.cs
--- a/Engine/Source/Programs/Horde/Plugins/Compute/HordeServer.Compute/ComputeConfig.cs
+++ b/Engine/Source/Programs/Horde/Plugins/Compute/HordeServer.Compute/ComputeConfig.cs
@@ -8,6 +8,7 @@ using EpicGames.Horde.Acls;
 using EpicGames.Horde.Agents.Pools;
 using EpicGames.Horde.Common;
 using EpicGames.Horde.Compute;
+using EpicGames.Horde.Telemetry;
 using EpicGames.Horde.Tools;
 using EpicGames.Serialization;
 using HordeServer.Acls;
@@ -40,6 +41,11 @@ namespace HordeServer
 		/// </summary>
 		public ConfigVersion VersionEnum { get; private set; }
 
+		/// <summary>
+		/// Telemetry store for lease and compute telemetry data
+		/// </summary>
+		public TelemetryStoreId TelemetryStoreId { get; set; }
+
 		/// <summary>
 		/// List of costs of a particular agent type
 		/// </summary>
diff --git a/Engine/Source/Programs/Horde/Plugins/Compute/HordeServer.Compute/ComputePlugin.cs b/Engine/Source/Programs/Horde/Plugins/Compute/HordeServer.Compute/ComputePlugin.cs
--- a/Engine/Source/Programs/Horde/Plugins/Compute/HordeServer.Compute/ComputePlugin.cs
+++ b/Engine/Source/Programs/Horde/Plugins/Compute/HordeServer.Compute/ComputePlugin.cs
@@ -83,6 +83,7 @@ namespace HordeServer
 			services.AddSingleton<IAgentVersionProvider, AgentVersionProvider>();
 
 			services.AddSingleton<AgentService>();
+			services.AddHostedService<LeaseService>();
 			services.AddSingleton(provider => new Lazy<AgentService>(provider.GetRequiredService<AgentService>));
 
 			services.AddSingleton<AwsAutoScalingLifecycleService>();
diff --git a/Engine/Source/Programs/Horde/Plugins/Compute/HordeServer.Compute/HordeServer.Compute.csproj b/Engine/Source/Programs/Horde/Plugins/Compute/HordeServer.Compute/HordeServer.Compute.csproj
--- a/Engine/Source/Programs/Horde/Plugins/Compute/HordeServer.Compute/HordeServer.Compute.csproj
+++ b/Engine/Source/Programs/Horde/Plugins/Compute/HordeServer.Compute/HordeServer.Compute.csproj
@@ -66,6 +66,7 @@
   </ItemGroup>
 
   <ItemGroup>
+    <ProjectReference Include="..\..\..\..\Shared\EpicGames.Analytics.Telemetry\EpicGames.Analytics.Telemetry.csproj" />
     <ProjectReference Include="..\..\..\..\Shared\EpicGames.Core\EpicGames.Core.csproj" />
     <ProjectReference Include="..\..\..\..\Shared\EpicGames.Horde\EpicGames.Horde.csproj" />
     <ProjectReference Include="..\..\..\..\Shared\EpicGames.IoHash\EpicGames.IoHash.csproj" />
