Make Test Health header sticky

#rnx
#Synced-CL 50905358

diff --git a/Engine/Source/Programs/Horde/HordeDashboard/plugins/test-automation/testAutomationHealth.tsx b/Engine/Source/Programs/Horde/HordeDashboard/plugins/test-automation/testAutomationHealth.tsx
--- a/Engine/Source/Programs/Horde/HordeDashboard/plugins/test-automation/testAutomationHealth.tsx
+++ b/Engine/Source/Programs/Horde/HordeDashboard/plugins/test-automation/testAutomationHealth.tsx
@@ -1,6 +1,6 @@
 // Copyright Epic Games, Inc. All Rights Reserved.
 
-import { DetailsList, FontIcon, IColumn, IconButton, Label, Modal, SelectionMode, Stack, Text, HoverCard, HoverCardType, IPlainCardProps, DirectionalHint, DetailsListLayoutMode, Spinner, SpinnerSize, Toggle, DatePicker, ComboBox, IComboBoxOption, FocusZone, FocusZoneTabbableElements, FocusTrapCallout, PrimaryButton, DefaultButton, TooltipHost, TooltipOverflowMode, TooltipDelay, Link } from "@fluentui/react";
+import { DetailsList, FontIcon, IColumn, IconButton, Label, Modal, SelectionMode, Stack, Text, HoverCard, HoverCardType, IPlainCardProps, DirectionalHint, DetailsListLayoutMode, Spinner, SpinnerSize, Toggle, DatePicker, ComboBox, IComboBoxOption, FocusZone, FocusZoneTabbableElements, FocusTrapCallout, PrimaryButton, DefaultButton, TooltipHost, TooltipOverflowMode, TooltipDelay, Link, IRenderFunction, IDetailsHeaderProps, Sticky, StickyPositionType, ScrollablePane } from "@fluentui/react";
 import { MetadataRef, TestDataHandler, TestNameRef, TestSessionResult } from "./testData";
 import { useCallback, useEffect, useMemo, useState } from "react";
 import { observer } from "mobx-react-lite";
@@ -1039,6 +1039,17 @@ const TestHealthDetailsModal: React.FC<{ test: TestHealthItem, handler: TestData
 
 const getTestHealthItemKey = (item: TestHealthItem) => `test-health-${item.nameRef.id}`;
 
+const onRenderDetailsHeader: IRenderFunction<IDetailsHeaderProps> = (props, defaultRender) => {
+    if (!defaultRender) {
+        return null;
+    }
+    return (
+        <Sticky stickyPosition={StickyPositionType.Header}>
+            {defaultRender(props)}
+        </Sticky>
+    );
+};
+
 export const TestAutomationHealth: React.FC<{ handler: TestDataHandler}> = observer(({ handler }) => {
     const [healthDetailsState, setHealthDetailsState] = useState<{test?: TestHealthItem}>({});
     const [healthListState, setHealthListState] = useState<TestHealthListState>({items: [], columns: []});
@@ -1111,15 +1122,18 @@ export const TestAutomationHealth: React.FC<{ handler: TestDataHandler}> = obser
                     />
                 }
                 { !handler.queryAuditsLoading &&
-                    <DetailsList
-                        items={healthListState.items}
-                        getKey={getTestHealthItemKey}
-                        compact={true}
-                        columns={healthListState.columns}
-                        isHeaderVisible={true}
-                        selectionMode={SelectionMode.none}
-                        layoutMode={DetailsListLayoutMode.justified}
-                    />
+                    <ScrollablePane>
+                        <DetailsList
+                            items={healthListState.items}
+                            getKey={getTestHealthItemKey}
+                            compact={true}
+                            columns={healthListState.columns}
+                            isHeaderVisible={true}
+                            onRenderDetailsHeader={onRenderDetailsHeader}
+                            selectionMode={SelectionMode.none}
+                            layoutMode={DetailsListLayoutMode.justified}
+                        />
+                    </ScrollablePane>
                 }
                 { handler.queryAuditsLoading &&
                     <Stack horizontalAlign='center' style={{ width: "100%", position: "absolute", top: 24 }} tokens={{childrenGap: 8}}>
