Horde: V3 pool migration rollout hardening
Fix shadow pool scaling state loss in PoolCollection.GetAsync
Add V3 guard in PoolUpdateService to skip config-managed pools
Enhance migration status endpoint with rich response
Add config diff endpoint and CLI script

diff --git a/Engine/Source/Programs/Horde/HordeServer.Shared/Configuration/IChangelistFileReader.cs b/Engine/Source/Programs/Horde/HordeServer.Shared/Configuration/IChangelistFileReader.cs
new file mode 100644
--- /dev/null
+++ b/Engine/Source/Programs/Horde/HordeServer.Shared/Configuration/IChangelistFileReader.cs
@@ -0,0 +1,32 @@
+// Copyright Epic Games, Inc. All Rights Reserved.
+
+namespace HordeServer.Configuration
+{
+	/// <summary>
+	/// Reads files from a shelved Perforce changelist for config validation and diffing.
+	/// </summary>
+	public interface IChangelistFileReader
+	{
+		/// <summary>
+		/// Reads all config-like files from a shelved changelist.
+		/// </summary>
+		/// <param name="changelist">The changelist number (must be shelved)</param>
+		/// <param name="cluster">Perforce cluster name (null for default)</param>
+		/// <param name="cancellationToken">Cancellation token</param>
+		/// <returns>Dictionary mapping config URI to file contents</returns>
+		Task<Dictionary<Uri, byte[]>> ReadChangelistFilesAsync(int changelist, string? cluster, CancellationToken cancellationToken);
+	}
+
+	/// <summary>
+	/// Exception thrown when a changelist cannot be found or described.
+	/// </summary>
+	public class ChangelistNotFoundException : Exception
+	{
+		/// <summary>
+		/// Constructor
+		/// </summary>
+		public ChangelistNotFoundException(string message) : base(message)
+		{
+		}
+	}
+}
diff --git a/Engine/Source/Programs/Horde/HordeServer.Shared/Configuration/IConfigService.cs b/Engine/Source/Programs/Horde/HordeServer.Shared/Configuration/IConfigService.cs
--- a/Engine/Source/Programs/Horde/HordeServer.Shared/Configuration/IConfigService.cs
+++ b/Engine/Source/Programs/Horde/HordeServer.Shared/Configuration/IConfigService.cs
@@ -1,5 +1,7 @@
 ﻿// Copyright Epic Games, Inc. All Rights Reserved.
 
+using HordeServer.Plugins;
+
 namespace HordeServer.Configuration
 {
 	/// <summary>
@@ -16,5 +18,41 @@ namespace HordeServer.Configuration
 		/// Validate a new set of config files. Parses and runs PostLoad methods on them.
 		/// </summary>
 		Task<string?> ValidateAsync(Dictionary<Uri, byte[]> files, CancellationToken cancellationToken);
+
+		/// <summary>
+		/// Resolve a set of config file overrides through the full pipeline (includes, macros, inheritance)
+		/// and return the resulting GlobalConfig. Unread files are returned as warnings, not errors.
+		/// </summary>
+		/// <param name="files">Config file overrides (URI to contents)</param>
+		/// <param name="cancellationToken">Cancellation token</param>
+		/// <returns>Resolved config, error message if parsing failed, list of unread file URIs</returns>
+		Task<ConfigResolveResult> ResolveWithOverridesAsync(Dictionary<Uri, byte[]> files, CancellationToken cancellationToken);
+	}
+
+	/// <summary>
+	/// Result from resolving config with overrides
+	/// </summary>
+	public class ConfigResolveResult
+	{
+		/// <summary>
+		/// The resolved global config (null if validation failed)
+		/// </summary>
+		public object? Config { get; set; }
+
+		/// <summary>
+		/// The resolved plugin configs (null if validation failed).
+		/// Plugins can use their GetXxxConfig() extension methods on this.
+		/// </summary>
+		public IDictionary<PluginName, IPluginConfig>? Plugins { get; set; }
+
+		/// <summary>
+		/// Error message if config parsing failed (null on success)
+		/// </summary>
+		public string? Error { get; set; }
+
+		/// <summary>
+		/// Config file URIs from the overrides that were not consumed by the pipeline
+		/// </summary>
+		public List<Uri> UnreadFiles { get; set; } = new List<Uri>();
 	}
 }
diff --git a/Engine/Source/Programs/Horde/HordeServer/Configuration/ConfigService.cs b/Engine/Source/Programs/Horde/HordeServer/Configuration/ConfigService.cs
--- a/Engine/Source/Programs/Horde/HordeServer/Configuration/ConfigService.cs
+++ b/Engine/Source/Programs/Horde/HordeServer/Configuration/ConfigService.cs
@@ -269,10 +269,7 @@ namespace HordeServer.Configuration
 			return context;
 		}
 
-		/// <summary>
-		/// Validate a new set of config files. Parses and runs PostLoad methods on them.
-		/// </summary>
-		public async Task<string?> ValidateAsync(Dictionary<Uri, byte[]> files, CancellationToken cancellationToken)
+		(Dictionary<Uri, OverrideConfigFile> OverrideFiles, ConfigContext Context) CreateOverrideContext(Dictionary<Uri, byte[]> files)
 		{
 			Dictionary<Uri, OverrideConfigFile> overrideFiles = new Dictionary<Uri, OverrideConfigFile>();
 			foreach ((Uri uri, byte[] data) in files)
@@ -287,12 +284,43 @@ namespace HordeServer.Configuration
 			}
 
 			ConfigContext context = CreateConfigContext(overrideSources, NullLogger.Instance);
-			try
+			return (overrideFiles, context);
+		}
+
+		async Task<GlobalConfig> ReadAndPostLoadOverridesAsync(ConfigContext context, CancellationToken cancellationToken)
+		{
+			GlobalConfig globalConfig = await ReadGlobalConfigAsync(context, cancellationToken);
+			globalConfig.PostLoad(_serverSettings, _pluginCollection.LoadedPlugins, _defaultAclModifiers, _logger);
+			await ResolveConfigSecretsAsync(globalConfig, cancellationToken);
+			return globalConfig;
+		}
+
+		string FormatConfigException(Exception ex, ConfigContext context, string logMessage)
+		{
+			StringBuilder message = new StringBuilder(ex.Message);
+			message.Append("\n\nLocation:\n");
+			foreach (IConfigFile include in context.IncludeStack)
 			{
-				GlobalConfig globalConfig = await ReadGlobalConfigAsync(context, cancellationToken);
-				globalConfig.PostLoad(_serverSettings, _pluginCollection.LoadedPlugins, _defaultAclModifiers, _logger);
+				string path = include.GetUserFormattedPath();
+				message.Append($"  {path}\n");
+			}
+			if (ex is not ConfigException && ex is not JsonException)
+			{
+				message.Append($"\n\nTrace:\n{ex.StackTrace}");
+				_logger.LogWarning(ex, "{Message}: {ExMessage}", logMessage, ex.Message);
+			}
+			return message.ToString();
+		}
 
-				await ResolveConfigSecretsAsync(globalConfig, cancellationToken);
+		/// <summary>
+		/// Validate a new set of config files. Parses and runs PostLoad methods on them.
+		/// </summary>
+		public async Task<string?> ValidateAsync(Dictionary<Uri, byte[]> files, CancellationToken cancellationToken)
+		{
+			(Dictionary<Uri, OverrideConfigFile> overrideFiles, ConfigContext context) = CreateOverrideContext(files);
+			try
+			{
+				await ReadAndPostLoadOverridesAsync(context, cancellationToken);
 
 				foreach (OverrideConfigFile file in overrideFiles.Values)
 				{
@@ -306,19 +334,33 @@ namespace HordeServer.Configuration
 			}
 			catch (Exception ex)
 			{
-				StringBuilder message = new StringBuilder(ex.Message);
-				message.Append("\n\nLocation:\n");
-				foreach (IConfigFile include in context.IncludeStack)
-				{
-					string path = include.GetUserFormattedPath();
-					message.Append($"  {path}\n");
-				}
-				if (ex is not ConfigException && ex is not JsonException)
+				return FormatConfigException(ex, context, "Unhandled exception while validating config files");
+			}
+		}
+
+		/// <inheritdoc/>
+		public async Task<ConfigResolveResult> ResolveWithOverridesAsync(Dictionary<Uri, byte[]> files, CancellationToken cancellationToken)
+		{
+			(Dictionary<Uri, OverrideConfigFile> overrideFiles, ConfigContext context) = CreateOverrideContext(files);
+			try
+			{
+				GlobalConfig globalConfig = await ReadAndPostLoadOverridesAsync(context, cancellationToken);
+
+				// Collect unread files as warnings (not errors)
+				List<Uri> unreadFiles = new List<Uri>();
+				foreach (OverrideConfigFile file in overrideFiles.Values)
 				{
-					message.Append($"\n\nTrace:\n{ex.StackTrace}");
-					_logger.LogWarning(ex, "Unhandled exception while validating config files: {Message}", ex.Message);
+					if (!file.WasRead)
+					{
+						unreadFiles.Add(file.Uri);
+					}
 				}
-				return message.ToString();
+
+				return new ConfigResolveResult { Config = globalConfig, Plugins = globalConfig.Plugins, UnreadFiles = unreadFiles };
+			}
+			catch (Exception ex)
+			{
+				return new ConfigResolveResult { Error = FormatConfigException(ex, context, "Unhandled exception while resolving config overrides") };
 			}
 		}
 
diff --git a/Engine/Source/Programs/Horde/Plugins/Build/HordeServer.Build.Tests/Agents/Pools/PoolUpdateServiceTest.cs b/Engine/Source/Programs/Horde/Plugins/Build/HordeServer.Build.Tests/Agents/Pools/PoolUpdateServiceTest.cs
--- a/Engine/Source/Programs/Horde/Plugins/Build/HordeServer.Build.Tests/Agents/Pools/PoolUpdateServiceTest.cs
+++ b/Engine/Source/Programs/Horde/Plugins/Build/HordeServer.Build.Tests/Agents/Pools/PoolUpdateServiceTest.cs
@@ -1,6 +1,7 @@
 // Copyright Epic Games, Inc. All Rights Reserved.
 
 using EpicGames.Horde.Agents;
+using EpicGames.Horde.Agents.Pools;
 using HordeServer.Agents;
 using HordeServer.Agents.Pools;
 using HordeServer.Projects;
@@ -24,7 +25,7 @@ public class PoolUpdateServiceTest : BuildTestSetup
 	public async Task SetupAsync()
 	{
 		await UpdateConfigAsync(x => x.Plugins.GetComputeConfig().Pools.Clear());
-		_pus = new PoolUpdateService(AgentService, PoolService, PoolCollection, Clock, ServiceProvider.GetRequiredService<IOptionsMonitor<BuildConfig>>(), Tracer, new NullLogger<PoolUpdateService>());
+		_pus = new PoolUpdateService(AgentService, PoolService, PoolCollection, Clock, ServiceProvider.GetRequiredService<IOptionsMonitor<BuildConfig>>(), ServiceProvider.GetRequiredService<IOptionsMonitor<ComputeConfig>>(), Tracer, new NullLogger<PoolUpdateService>());
 		_pool = await CreatePoolAsync(new PoolConfig { Name = "testPool", EnableAutoscaling = true });
 		_enabledAgent = await CreateAgentAsync(_pool, true);
 		_disabledAgent = await CreateAgentAsync(_pool, false, status: AgentStatus.Stopped);
@@ -143,16 +144,136 @@ public class PoolUpdateServiceTest : BuildTestSetup
 			agentTypes[agentTypeName] = new AgentConfig { Pool = _pool.Id, Workspace = workspaceId };
 			workspaceTypes[workspaceId] = new WorkspaceConfig { Identifier = workspaceId, ConformDiskFreeSpace = autoConformThresholdsMb[i] };
 		}
-		
+
 		await UpdateConfigAsync(x =>
 		{
 			x.Plugins.GetBuildConfig().Projects.Clear();
 			x.Plugins.GetBuildConfig().Projects.Add(new ProjectConfig { Streams = [new StreamConfig { WorkspaceTypes = workspaceTypes, AgentTypes = agentTypes }]});
 		});
-		
+
 		IAgent? agent = await CreateAgentAsync(_pool, properties: [$"{KnownPropertyNames.WorkingDiskFreeSpace}={freeDiskSpaceMb * 1024 * 1024}"]);
 		agent = await agent.TryTerminateSessionAsync(); // Make agent status = stopped
 		Assert.IsNotNull(agent);
 		return agent;
 	}
+
+	[TestMethod]
+	public async Task UpdatePools_V3_SkipsConfigManagedPoolAsync()
+	{
+		// Arrange: enable V3 and create a config-managed pool with a stream referencing it
+		await UpdateConfigAsync(config =>
+		{
+			config.Version = (int)ConfigVersion.IncrementalPoolMigration;
+			ComputeConfig computeConfig = config.Plugins.GetComputeConfig();
+			// _pool was already created via config in SetupAsync
+			config.Plugins.GetBuildConfig().Projects.Clear();
+			config.Plugins.GetBuildConfig().Projects.Add(new ProjectConfig
+			{
+				Streams =
+				[
+					new StreamConfig
+					{
+						WorkspaceTypes = new Dictionary<string, WorkspaceConfig>
+						{
+							["ws1"] = new WorkspaceConfig { Identifier = "ws1" }
+						},
+						AgentTypes = new Dictionary<string, AgentConfig>
+						{
+							["at1"] = new AgentConfig { Pool = _pool.Id, Workspace = "ws1" }
+						}
+					}
+				]
+			});
+		});
+
+		// Act: run the ticker
+		await _pus.UpdatePoolsAsync(CancellationToken.None);
+
+		// Assert: the config-managed pool should NOT have a V1 DB document created or updated.
+		// Since the pool was created via config (not DB), no V1 doc should exist.
+		// The V3 guard prevents UpdateConfigAsync from being called for config-managed pools.
+		List<PoolId> dbPoolIds = await PoolCollection.GetDatabasePoolIdsAsync(CancellationToken.None);
+		Assert.IsFalse(dbPoolIds.Contains(_pool.Id), "Config-managed pool should not have a V1 DB document");
+	}
+
+	[TestMethod]
+	public async Task UpdatePools_V3_UpdatesDbOnlyPoolAsync()
+	{
+		// Arrange: enable V3 first so DB pools are visible via GetAsync (default V2 skips DB)
+		await UpdateConfigAsync(config =>
+		{
+			config.Version = (int)ConfigVersion.IncrementalPoolMigration;
+		});
+
+		PoolConfig dbPoolConfig = new PoolConfig { Name = "DbOnlyPool" };
+		IPool dbPool = await CreatePoolAsync(dbPoolConfig, viaConfig: false);
+
+		await UpdateConfigAsync(config =>
+		{
+			config.Version = (int)ConfigVersion.IncrementalPoolMigration;
+			config.Plugins.GetBuildConfig().Projects.Clear();
+			config.Plugins.GetBuildConfig().Projects.Add(new ProjectConfig
+			{
+				Streams =
+				[
+					new StreamConfig
+					{
+						WorkspaceTypes = new Dictionary<string, WorkspaceConfig>
+						{
+							["ws2"] = new WorkspaceConfig { Identifier = "ws2" }
+						},
+						AgentTypes = new Dictionary<string, AgentConfig>
+						{
+							["at2"] = new AgentConfig { Pool = dbPool.Id, Workspace = "ws2" }
+						}
+					}
+				]
+			});
+		});
+
+		// Act
+		await _pus.UpdatePoolsAsync(CancellationToken.None);
+
+		// Assert: the DB-only pool should have its workspaces updated
+		IPool? updatedPool = await PoolCollection.GetAsync(dbPool.Id);
+		Assert.IsNotNull(updatedPool);
+		Assert.AreEqual(1, updatedPool.Workspaces.Count);
+		Assert.AreEqual("ws2", updatedPool.Workspaces[0].Identifier);
+	}
+
+	[TestMethod]
+	public async Task UpdatePools_V1_UpdatesAllPoolsRegardlessAsync()
+	{
+		// Arrange: V1 mode (default) — pool update service should update all pools
+		await UpdateConfigAsync(config =>
+		{
+			config.Version = (int)ConfigVersion.Initial;
+			config.Plugins.GetBuildConfig().Projects.Clear();
+			config.Plugins.GetBuildConfig().Projects.Add(new ProjectConfig
+			{
+				Streams =
+				[
+					new StreamConfig
+					{
+						WorkspaceTypes = new Dictionary<string, WorkspaceConfig>
+						{
+							["ws3"] = new WorkspaceConfig { Identifier = "ws3" }
+						},
+						AgentTypes = new Dictionary<string, AgentConfig>
+						{
+							["at3"] = new AgentConfig { Pool = _pool.Id, Workspace = "ws3" }
+						}
+					}
+				]
+			});
+		});
+
+		// Act
+		await _pus.UpdatePoolsAsync(CancellationToken.None);
+
+		// Assert: at V1, all pools should be updated regardless
+		IPool? pool = await PoolCollection.GetAsync(_pool.Id);
+		Assert.IsNotNull(pool);
+		Assert.AreEqual(1, pool.Workspaces.Count);
+	}
 }
diff --git a/Engine/Source/Programs/Horde/Plugins/Build/HordeServer.Build.Tests/Agents/Pools/PoolsControllerTest.cs b/Engine/Source/Programs/Horde/Plugins/Build/HordeServer.Build.Tests/Agents/Pools/PoolsControllerTest.cs
--- a/Engine/Source/Programs/Horde/Plugins/Build/HordeServer.Build.Tests/Agents/Pools/PoolsControllerTest.cs
+++ b/Engine/Source/Programs/Horde/Plugins/Build/HordeServer.Build.Tests/Agents/Pools/PoolsControllerTest.cs
@@ -1,8 +1,14 @@
 ﻿// Copyright Epic Games, Inc. All Rights Reserved.
 
+using EpicGames.Horde.Agents.Pools;
 using HordeServer.Agents.Fleet;
 using HordeServer.Agents.Pools;
 using Microsoft.AspNetCore.Mvc;
+using Microsoft.Extensions.DependencyInjection;
+using Microsoft.Extensions.Options;
+using MongoDB.Bson;
+using MongoDB.Bson.Serialization.Attributes;
+using MongoDB.Driver;
 
 namespace HordeServer.Tests.Agents.Pools
 {
@@ -79,5 +85,110 @@ namespace HordeServer.Tests.Agents.Pools
 			Assert.AreEqual(poolConfig.FleetManagers[0].Condition!.Text, pool.FleetManagers[0].Condition!.Text);
 			Assert.AreEqual(poolConfig.FleetManagers[0].Config, pool.FleetManagers[0].Config);
 		}
+
+		[TestMethod]
+		public async Task ShadowPool_V3_CopiesV1ScalingStateToV2Async()
+		{
+			// Arrange: insert a V1 pool doc directly with scaling state, and ensure no V2 exists.
+			// Then add the same pool ID to config at V3 (shadow pool).
+			// When GetAsync is called, it should load V1 scaling state into the new V2 entry.
+			PoolId poolId = new PoolId("shadow-scaling-test");
+			DateTime scaleUpTime = new DateTime(2025, 1, 15, 10, 0, 0, DateTimeKind.Utc);
+			DateTime scaleDownTime = new DateTime(2025, 1, 15, 8, 0, 0, DateTimeKind.Utc);
+
+			// Insert V1 doc with scaling state directly via MongoDB
+			IMongoCollection<BsonDocument> poolsV1 = MongoService.GetCollection<BsonDocument>("Pools");
+			await poolsV1.InsertOneAsync(new BsonDocument
+			{
+				{ "_id", poolId.ToString() },
+				{ "Name", "ShadowScalingTest" },
+				{ "EnableAutoscaling", true },
+				{ "LastScaleUpTime", scaleUpTime },
+				{ "LastScaleDownTime", scaleDownTime },
+				{ "LastAgentCount", 42 },
+				{ "LastDesiredAgentCount", 50 },
+				{ "UpdateIndex", 0 },
+				{ "Workspaces", new BsonArray() },
+				{ "Properties", new BsonDocument() },
+			});
+
+			// Enable V3 and add same pool to config (making it a shadow pool)
+			await UpdateConfigAsync(config =>
+			{
+				config.Version = (int)ConfigVersion.IncrementalPoolMigration;
+				ComputeConfig computeConfig = config.Plugins.GetComputeConfig();
+				computeConfig.Pools.Add(new PoolConfig { Id = poolId, Name = "ShadowScalingTest" });
+			});
+
+			// Act: get the pool — should create V2 entry with V1 scaling state
+			IPool? shadowPool = await PoolCollection.GetAsync(poolId);
+
+			// Assert
+			Assert.IsNotNull(shadowPool);
+			Assert.AreEqual(scaleUpTime, shadowPool.LastScaleUpTime);
+			Assert.AreEqual(scaleDownTime, shadowPool.LastScaleDownTime);
+			Assert.AreEqual(42, shadowPool.LastAgentCount);
+			Assert.AreEqual(50, shadowPool.LastDesiredAgentCount);
+		}
+
+		[TestMethod]
+		public async Task ShadowPool_V3_ExistingV2NotOverwrittenAsync()
+		{
+			// Arrange: enable V3 first so DB pools are visible via GetAsync (default V2 skips DB)
+			await UpdateConfigAsync(config =>
+			{
+				config.Version = (int)ConfigVersion.IncrementalPoolMigration;
+			});
+
+			// Create a DB pool and get it once to create V2 entry
+			PoolConfig dbPoolConfig = new PoolConfig { Name = "ExistingV2Pool" };
+			IPool dbPool = await CreatePoolAsync(dbPoolConfig, viaConfig: false);
+
+			// TryUpdateAsync updates the V2 entry (V2 was already created by CreatePoolAsync->GetAsync)
+			DateTime scaleUpTime = new DateTime(2025, 2, 1, 12, 0, 0, DateTimeKind.Utc);
+			dbPool = await dbPool.TryUpdateAsync(new UpdatePoolOptions
+			{
+				LastScaleUpTime = scaleUpTime,
+				LastAgentCount = 100
+			}, CancellationToken.None) ?? throw new Exception("Update failed");
+
+			// Now add the same pool to config (shadow pool)
+			await UpdateConfigAsync(config =>
+			{
+				config.Version = (int)ConfigVersion.IncrementalPoolMigration;
+				ComputeConfig computeConfig = config.Plugins.GetComputeConfig();
+				computeConfig.Pools.Add(new PoolConfig { Id = dbPoolConfig.Id, Name = "ExistingV2Pool" });
+			});
+
+			// Act: get the pool — V2 already exists, should NOT be overwritten
+			IPool? pool = await PoolCollection.GetAsync(dbPoolConfig.Id);
+
+			// Assert: V2 entry should have the existing state
+			Assert.IsNotNull(pool);
+			Assert.AreEqual(scaleUpTime, pool.LastScaleUpTime);
+			Assert.AreEqual(100, pool.LastAgentCount);
+		}
+
+		[TestMethod]
+		public async Task ConfigOnlyPool_V3_CreatesEmptyV2Async()
+		{
+			// Arrange: config-only pool (no V1 doc in DB at all)
+			PoolId configOnlyPoolId = new PoolId("configonlypool");
+			await UpdateConfigAsync(config =>
+			{
+				config.Version = (int)ConfigVersion.IncrementalPoolMigration;
+				ComputeConfig computeConfig = config.Plugins.GetComputeConfig();
+				computeConfig.Pools.Add(new PoolConfig { Id = configOnlyPoolId, Name = "ConfigOnlyPool" });
+			});
+
+			// Act: get the pool — should create V2 entry with empty scaling state, no exception
+			IPool? pool = await PoolCollection.GetAsync(configOnlyPoolId);
+
+			// Assert
+			Assert.IsNotNull(pool);
+			Assert.IsNull(pool.LastScaleUpTime);
+			Assert.IsNull(pool.LastScaleDownTime);
+			Assert.IsNull(pool.LastAgentCount);
+		}
 	}
 }
\ No newline at end of file
diff --git a/Engine/Source/Programs/Horde/Plugins/Build/HordeServer.Build/Agents/Pools/PoolUpdateService.cs b/Engine/Source/Programs/Horde/Plugins/Build/HordeServer.Build/Agents/Pools/PoolUpdateService.cs
--- a/Engine/Source/Programs/Horde/Plugins/Build/HordeServer.Build/Agents/Pools/PoolUpdateService.cs
+++ b/Engine/Source/Programs/Horde/Plugins/Build/HordeServer.Build/Agents/Pools/PoolUpdateService.cs
@@ -3,8 +3,8 @@
 using EpicGames.Core;
 using EpicGames.Horde.Agents;
 using EpicGames.Horde.Agents.Pools;
-using HordeServer.VersionControl.Perforce;
 using HordeServer.Streams;
+using HordeServer.VersionControl.Perforce;
 using HordeServer.Utilities;
 using Microsoft.Extensions.Hosting;
 using Microsoft.Extensions.Logging;
@@ -23,6 +23,7 @@ namespace HordeServer.Agents.Pools
 		readonly IPoolCollection _pools;
 		readonly IClock _clock;
 		readonly IOptionsMonitor<BuildConfig> _buildConfig;
+		readonly IOptionsMonitor<ComputeConfig> _computeConfig;
 		readonly ILogger<PoolUpdateService> _logger;
 		readonly Tracer _tracer;
 		readonly ITicker _updatePoolsTicker;
@@ -32,13 +33,14 @@ namespace HordeServer.Agents.Pools
 		/// <summary>
 		/// Constructor
 		/// </summary>
-		public PoolUpdateService(AgentService agentService, PoolService poolService, IPoolCollection pools, IClock clock, IOptionsMonitor<BuildConfig> buildConfig, Tracer tracer, ILogger<PoolUpdateService> logger)
+		public PoolUpdateService(AgentService agentService, PoolService poolService, IPoolCollection pools, IClock clock, IOptionsMonitor<BuildConfig> buildConfig, IOptionsMonitor<ComputeConfig> computeConfig, Tracer tracer, ILogger<PoolUpdateService> logger)
 		{
 			_agentService = agentService;
 			_poolService = poolService;
 			_pools = pools;
 			_clock = clock;
 			_buildConfig = buildConfig;
+			_computeConfig = computeConfig;
 			_tracer = tracer;
 			_logger = logger;
 			_updatePoolsTicker = clock.AddTicker($"{nameof(PoolUpdateService)}.{nameof(UpdatePoolsAsync)}", TimeSpan.FromSeconds(30.0), UpdatePoolsAsync, logger);
@@ -211,9 +213,19 @@ namespace HordeServer.Agents.Pools
 					}
 				}
 
+				// Capture compute config for V3 guard
+				ComputeConfig computeConfig = _computeConfig.CurrentValue;
+
 				// Update the list of workspaces for each pool
 				foreach (IPoolConfig currentPool in currentPools)
 				{
+					// V3+: skip config-managed pools — their workspaces come from config files
+					if (computeConfig.VersionEnum >= ConfigVersion.IncrementalPoolMigration
+						&& computeConfig.TryGetPool(currentPool.Id, out _))
+					{
+						continue;
+					}
+
 					// Get the new list of workspaces for this pool
 					List<AgentWorkspaceInfo>? newWorkspaces;
 					if (!poolToAgentWorkspaces.TryGetValue(currentPool.Id, out newWorkspaces))
diff --git a/Engine/Source/Programs/Horde/Plugins/Build/HordeServer.Build/BuildPlugin.cs b/Engine/Source/Programs/Horde/Plugins/Build/HordeServer.Build/BuildPlugin.cs
--- a/Engine/Source/Programs/Horde/Plugins/Build/HordeServer.Build/BuildPlugin.cs
+++ b/Engine/Source/Programs/Horde/Plugins/Build/HordeServer.Build/BuildPlugin.cs
@@ -140,20 +140,6 @@ namespace HordeServer
 			// Issues need to be assigned asynchronously on any pod
 			services.AddHostedService(provider => provider.GetRequiredService<IssueService>());
 
-			if (_serverInfo.IsRunModeActive(RunMode.Worker) && !_serverInfo.ReadOnlyMode)
-			{
-				services.AddHostedService<AgentReportService>();
-				services.AddHostedService<IssueReportService>();
-				services.AddHostedService<JobExpirationService>();
-				services.AddHostedService(provider => provider.GetRequiredService<PerforceLoadBalancer>());
-				services.AddHostedService<PoolUpdateService>();
-				services.AddHostedService<UtilizationDataService>();
-				services.AddHostedService(provider => provider.GetRequiredService<DeviceService>());
-				services.AddHostedService<DeviceReportService>();
-				services.AddHostedService(provider => provider.GetRequiredService<TestDataService>());
-				services.AddHostedService<TestHealthService>();
-			}
-
 			services.AddHostedService(provider => provider.GetRequiredService<IExternalIssueService>());
 
 			// Task sources. Order of registration is important here; it dictates the priority in which sources are served.
@@ -181,6 +167,7 @@ namespace HordeServer
 			}
 
 			services.AddSingleton<IVersionControlService>(sp => sp.GetRequiredService<IPerforceService>());
+			services.AddSingleton<IChangelistFileReader, ChangelistFileReader>();
 
 			services.AddSingleton<PerforceReplicator>();
 
diff --git a/Engine/Source/Programs/Horde/Plugins/Build/HordeServer.Build/VersionControl/Perforce/ChangelistFileReader.cs b/Engine/Source/Programs/Horde/Plugins/Build/HordeServer.Build/VersionControl/Perforce/ChangelistFileReader.cs
new file mode 100644
--- /dev/null
+++ b/Engine/Source/Programs/Horde/Plugins/Build/HordeServer.Build/VersionControl/Perforce/ChangelistFileReader.cs
@@ -0,0 +1,78 @@
+// Copyright Epic Games, Inc. All Rights Reserved.
+
+using EpicGames.Perforce;
+using HordeServer.Configuration;
+using Microsoft.Extensions.Logging;
+
+namespace HordeServer.VersionControl.Perforce
+{
+	/// <summary>
+	/// Reads config files from shelved Perforce changelists for diff/validation.
+	/// </summary>
+	public class ChangelistFileReader : IChangelistFileReader
+	{
+		static readonly List<string> s_configSuffixes = new List<string>
+		{
+			"/globals.json", "/global.json", ".project.json", ".stream.json",
+			".dashboard.json", ".telemetry.json"
+		};
+
+		readonly IPerforceService _perforceService;
+		readonly ILogger<ChangelistFileReader> _logger;
+
+		/// <summary>
+		/// Constructor
+		/// </summary>
+		public ChangelistFileReader(IPerforceService perforceService, ILogger<ChangelistFileReader> logger)
+		{
+			_perforceService = perforceService;
+			_logger = logger;
+		}
+
+		/// <inheritdoc/>
+		public async Task<Dictionary<Uri, byte[]>> ReadChangelistFilesAsync(
+			int changelist, string? cluster, CancellationToken cancellationToken)
+		{
+			cluster ??= "default";
+			using IPooledPerforceConnection perforce = await _perforceService.ConnectAsync(
+				cluster, cancellationToken: cancellationToken);
+
+			PerforceResponse<DescribeRecord> describeResponse = await perforce.TryDescribeAsync(
+				DescribeOptions.Shelved, -1, changelist, cancellationToken);
+			if (!describeResponse.Succeeded)
+			{
+				throw new ChangelistNotFoundException(
+					$"Changelist {changelist} not found or has no shelved files. Ensure the CL is shelved (not just pending).");
+			}
+
+			DescribeRecord record = describeResponse.Data;
+			Dictionary<Uri, byte[]> files = new Dictionary<Uri, byte[]>();
+
+			foreach (DescribeFileRecord fileRecord in record.Files)
+			{
+				// Filter to config-like files only
+				if (s_configSuffixes.FirstOrDefault(
+					suffix => fileRecord.DepotFile.EndsWith(suffix, StringComparison.OrdinalIgnoreCase)) == null)
+				{
+					continue;
+				}
+
+				PerforceResponse<PrintRecord<byte[]>> printResponse = await perforce.TryPrintDataAsync(
+					$"{fileRecord.DepotFile}@={changelist}", cancellationToken);
+				if (!printResponse.Succeeded || printResponse.Data.Contents == null)
+				{
+					_logger.LogWarning(
+						"Skipping file {DepotFile}@={Changelist}: unable to read contents",
+						fileRecord.DepotFile, changelist);
+					continue;
+				}
+
+				PrintRecord<byte[]> printRecord = printResponse.Data;
+				Uri uri = new Uri($"perforce://{cluster}{printRecord.DepotFile}");
+				files.Add(uri, printRecord.Contents);
+			}
+
+			return files;
+		}
+	}
+}
diff --git a/Engine/Source/Programs/Horde/Plugins/Compute/HordeServer.Compute/Agents/Pools/IPoolCollection.cs b/Engine/Source/Programs/Horde/Plugins/Compute/HordeServer.Compute/Agents/Pools/IPoolCollection.cs
--- a/Engine/Source/Programs/Horde/Plugins/Compute/HordeServer.Compute/Agents/Pools/IPoolCollection.cs
+++ b/Engine/Source/Programs/Horde/Plugins/Compute/HordeServer.Compute/Agents/Pools/IPoolCollection.cs
@@ -51,6 +51,13 @@ namespace HordeServer.Agents.Pools
 		/// </summary>
 		Task<List<PoolId>> GetDatabasePoolIdsAsync(CancellationToken cancellationToken = default);
 
+		/// <summary>
+		/// Gets all pool configs from the database (PoolsV1 collection) as IPoolConfig,
+		/// regardless of version settings. Used by migration status endpoint for field-by-field
+		/// comparison between config and DB pools.
+		/// </summary>
+		Task<IReadOnlyList<IPoolConfig>> GetDatabasePoolConfigsAsync(CancellationToken cancellationToken = default);
+
 		/// <summary>
 		/// Gets a pool by ID
 		/// </summary>
diff --git a/Engine/Source/Programs/Horde/Plugins/Compute/HordeServer.Compute/Agents/Pools/PoolCollection.cs b/Engine/Source/Programs/Horde/Plugins/Compute/HordeServer.Compute/Agents/Pools/PoolCollection.cs
--- a/Engine/Source/Programs/Horde/Plugins/Compute/HordeServer.Compute/Agents/Pools/PoolCollection.cs
+++ b/Engine/Source/Programs/Horde/Plugins/Compute/HordeServer.Compute/Agents/Pools/PoolCollection.cs
@@ -319,16 +319,37 @@ namespace HordeServer.Agents.Pools
 			PoolDocumentV2? state = await _poolsV2.Find(x => x.Id == id).FirstOrDefaultAsync(cancellationToken);
 			if (state == null)
 			{
-				state = new PoolDocumentV2 { Id = id };
-				if (poolConfig is PoolDocumentV1 poolV1)
+				PoolDocumentV2 newState = new PoolDocumentV2 { Id = id };
+
+				// Try to copy scaling state from V1 document.
+				// Use `as` cast so poolV1 is declared in outer scope for both paths.
+				PoolDocumentV1? poolV1 = poolConfig as PoolDocumentV1;
+				if (poolV1 == null
+					&& computeConfig.VersionEnum >= ConfigVersion.IncrementalPoolMigration)
+				{
+					// Shadow pool case: config came from file, but V1 doc may have scaling state
+					poolV1 = await _poolsV1.Find(x => x.Id == id).FirstOrDefaultAsync(cancellationToken);
+				}
+
+				if (poolV1 != null)
+				{
+					newState.LastScaleUpTime = poolV1.LastScaleUpTime;
+					newState.LastScaleDownTime = poolV1.LastScaleDownTime;
+					newState.LastScaleResult = poolV1.LastScaleResult;
+					newState.LastAgentCount = poolV1.LastAgentCount;
+					newState.LastDesiredAgentCount = poolV1.LastDesiredAgentCount;
+				}
+
+				if (!await _poolsV2.InsertOneIgnoreDuplicatesAsync(newState, cancellationToken))
 				{
-					state.LastScaleUpTime = poolV1.LastScaleUpTime;
-					state.LastScaleDownTime = poolV1.LastScaleDownTime;
-					state.LastScaleResult = poolV1.LastScaleResult;
-					state.LastAgentCount = poolV1.LastAgentCount;
-					state.LastDesiredAgentCount = poolV1.LastDesiredAgentCount;
+					// Another instance created the V2 doc first; read the actual persisted state
+					// to ensure UpdateIndex matches for optimistic concurrency in TryUpdateAsync.
+					state = await _poolsV2.Find(x => x.Id == id).FirstOrDefaultAsync(cancellationToken) ?? newState;
+				}
+				else
+				{
+					state = newState;
 				}
-				await _poolsV2.InsertOneIgnoreDuplicatesAsync(state, cancellationToken);
 			}
 
 			return new Pool(this, poolConfig, state);
@@ -375,6 +396,13 @@ namespace HordeServer.Agents.Pools
 			return results.Select(x => new PoolId(x.GetElement("_id").Value.AsString)).ToList();
 		}
 
+		/// <inheritdoc/>
+		public async Task<IReadOnlyList<IPoolConfig>> GetDatabasePoolConfigsAsync(CancellationToken cancellationToken)
+		{
+			List<PoolDocumentV1> docs = await _poolsV1.Find(FilterDefinition<PoolDocumentV1>.Empty).ToListAsync(cancellationToken);
+			return docs.Cast<IPoolConfig>().ToList();
+		}
+
 		/// <inheritdoc/>
 		public async Task<bool> DeleteConfigAsync(PoolId id, CancellationToken cancellationToken)
 		{
diff --git a/Engine/Source/Programs/Horde/Plugins/Compute/HordeServer.Compute/Server/PoolMigrationStatusResponse.cs b/Engine/Source/Programs/Horde/Plugins/Compute/HordeServer.Compute/Server/PoolMigrationStatusResponse.cs
new file mode 100644
--- /dev/null
+++ b/Engine/Source/Programs/Horde/Plugins/Compute/HordeServer.Compute/Server/PoolMigrationStatusResponse.cs
@@ -0,0 +1,164 @@
+// Copyright Epic Games, Inc. All Rights Reserved.
+
+namespace HordeServer.Server
+{
+	/// <summary>
+	/// Response from the pool migration status endpoint
+	/// </summary>
+	public class GetPoolMigrationStatusResponse
+	{
+		/// <summary>
+		/// Current config version number
+		/// </summary>
+		public int Version { get; set; }
+
+		/// <summary>
+		/// Migration progress summary
+		/// </summary>
+		public PoolMigrationProgressInfo Progress { get; set; } = new PoolMigrationProgressInfo();
+
+		/// <summary>
+		/// Actionable recommendations for operators
+		/// </summary>
+		public List<string> Recommendations { get; set; } = new List<string>();
+
+		/// <summary>
+		/// Shadow pools (exist in both config and DB)
+		/// </summary>
+		public List<ShadowPoolInfo> ShadowPools { get; set; } = new List<ShadowPoolInfo>();
+
+		/// <summary>
+		/// Pools that exist only in the database
+		/// </summary>
+		public List<DbOnlyPoolInfo> DbOnlyPools { get; set; } = new List<DbOnlyPoolInfo>();
+	}
+
+	/// <summary>
+	/// Migration progress counters
+	/// </summary>
+	public class PoolMigrationProgressInfo
+	{
+		/// <summary>
+		/// Total number of unique pools across config and DB
+		/// </summary>
+		public int TotalPools { get; set; }
+
+		/// <summary>
+		/// Pools managed via config files
+		/// </summary>
+		public int ConfigManaged { get; set; }
+
+		/// <summary>
+		/// Pools only in the database
+		/// </summary>
+		public int DbOnly { get; set; }
+
+		/// <summary>
+		/// Pools that exist in both config and DB
+		/// </summary>
+		public int Shadow { get; set; }
+
+		/// <summary>
+		/// Percentage of pools that are config-managed
+		/// </summary>
+		public double MigrationPercent { get; set; }
+	}
+
+	/// <summary>
+	/// Information about a shadow pool (exists in both config and DB)
+	/// </summary>
+	public class ShadowPoolInfo
+	{
+		/// <summary>
+		/// Pool identifier
+		/// </summary>
+		public string Id { get; set; } = string.Empty;
+
+		/// <summary>
+		/// Pool name
+		/// </summary>
+		public string Name { get; set; } = string.Empty;
+
+		/// <summary>
+		/// Number of agents assigned to this pool
+		/// </summary>
+		public int AgentCount { get; set; }
+
+		/// <summary>
+		/// Number of online agents
+		/// </summary>
+		public int OnlineAgents { get; set; }
+
+		/// <summary>
+		/// Pool classification (active, dead, suspicious)
+		/// </summary>
+		public string Classification { get; set; } = string.Empty;
+
+		/// <summary>
+		/// Whether the config and DB versions have identical field values
+		/// </summary>
+		public bool FieldsMatch { get; set; }
+
+		/// <summary>
+		/// Field-by-field differences between config and DB (only present when FieldsMatch is false)
+		/// </summary>
+		public Dictionary<string, FieldDiffEntry>? FieldDiffs { get; set; }
+
+		/// <summary>
+		/// Recommended action for this pool
+		/// </summary>
+		public string Action { get; set; } = string.Empty;
+	}
+
+	/// <summary>
+	/// Information about a DB-only pool (not in config files)
+	/// </summary>
+	public class DbOnlyPoolInfo
+	{
+		/// <summary>
+		/// Pool identifier
+		/// </summary>
+		public string Id { get; set; } = string.Empty;
+
+		/// <summary>
+		/// Pool name
+		/// </summary>
+		public string Name { get; set; } = string.Empty;
+
+		/// <summary>
+		/// Number of agents assigned to this pool
+		/// </summary>
+		public int AgentCount { get; set; }
+
+		/// <summary>
+		/// Number of online agents
+		/// </summary>
+		public int OnlineAgents { get; set; }
+
+		/// <summary>
+		/// Pool classification (active, dead, suspicious)
+		/// </summary>
+		public string Classification { get; set; } = string.Empty;
+
+		/// <summary>
+		/// Recommended action for this pool
+		/// </summary>
+		public string Action { get; set; } = string.Empty;
+	}
+
+	/// <summary>
+	/// A single field difference between config and DB values
+	/// </summary>
+	public class FieldDiffEntry
+	{
+		/// <summary>
+		/// Value from config
+		/// </summary>
+		public object? Config { get; set; }
+
+		/// <summary>
+		/// Value from database
+		/// </summary>
+		public object? Db { get; set; }
+	}
+}
diff --git a/Engine/Source/Programs/Horde/Plugins/Compute/HordeServer.Compute/Server/PoolsServerController.cs b/Engine/Source/Programs/Horde/Plugins/Compute/HordeServer.Compute/Server/PoolsServerController.cs
--- a/Engine/Source/Programs/Horde/Plugins/Compute/HordeServer.Compute/Server/PoolsServerController.cs
+++ b/Engine/Source/Programs/Horde/Plugins/Compute/HordeServer.Compute/Server/PoolsServerController.cs
@@ -1,8 +1,13 @@
 // Copyright Epic Games, Inc. All Rights Reserved.
 
+using System.Text.Json;
+using EpicGames.Horde;
+using EpicGames.Horde.Agents;
 using EpicGames.Horde.Agents.Pools;
 using HordeServer.Agents;
 using HordeServer.Agents.Pools;
+using HordeServer.Configuration;
+
 using HordeServer.Utilities;
 using Microsoft.AspNetCore.Authorization;
 using Microsoft.AspNetCore.Mvc;
@@ -20,6 +25,8 @@ namespace HordeServer.Server
 	[Route("[controller]")]
 	public class PoolsServerController : HordeControllerBase
 	{
+		static readonly JsonSerializerOptions s_compactJsonOptions = new JsonSerializerOptions { WriteIndented = false };
+
 		readonly IServiceProvider _serviceProvider;
 		readonly IOptionsSnapshot<ComputeConfig> _computeConfig;
 		readonly ILogger<PoolsServerController> _logger;
@@ -126,11 +133,11 @@ namespace HordeServer.Server
 		}
 
 		/// <summary>
-		/// Returns pool migration status: how many pools are in config, DB-only, and shadowed.
+		/// Returns pool migration status with per-pool details, field diffs, and recommendations.
 		/// </summary>
 		[HttpGet]
 		[Route("/api/v1/server/migrate/pool-status")]
-		public async Task<ActionResult<object>> GetPoolMigrationStatusAsync(CancellationToken cancellationToken)
+		public async Task<ActionResult<GetPoolMigrationStatusResponse>> GetPoolMigrationStatusAsync(CancellationToken cancellationToken)
 		{
 			if (!_computeConfig.Value.Authorize(PoolAclAction.ListPools, User))
 			{
@@ -138,28 +145,470 @@ namespace HordeServer.Server
 			}
 
 			ComputeConfig computeConfig = _computeConfig.Value;
-			HashSet<PoolId> configPoolIds = new HashSet<PoolId>(computeConfig.Pools.Select(x => x.Id));
+			Dictionary<PoolId, PoolConfig> configPoolLookup = computeConfig.Pools.GroupBy(x => x.Id).ToDictionary(g => g.Key, g => g.First());
 
 			IPoolCollection poolCollection = _serviceProvider.GetRequiredService<IPoolCollection>();
-			List<PoolId> dbPoolIdList = await poolCollection.GetDatabasePoolIdsAsync(cancellationToken);
-			HashSet<PoolId> dbPoolIds = new HashSet<PoolId>(dbPoolIdList);
+			IReadOnlyList<IPoolConfig> dbPools = await poolCollection.GetDatabasePoolConfigsAsync(cancellationToken);
+			Dictionary<PoolId, IPoolConfig> dbPoolLookup = dbPools.GroupBy(x => x.Id).ToDictionary(g => g.Key, g => g.First());
+
+			// Get agent stats for classification
+			Dictionary<PoolId, (int Total, int Online)> agentStats = await GetAgentStatsAsync(cancellationToken);
+
+			// Classify pools
+			HashSet<PoolId> shadowIds = new HashSet<PoolId>(configPoolLookup.Keys);
+			shadowIds.IntersectWith(dbPoolLookup.Keys);
+
+			HashSet<PoolId> dbOnlyIds = new HashSet<PoolId>(dbPoolLookup.Keys);
+			dbOnlyIds.ExceptWith(configPoolLookup.Keys);
+
+			// Build shadow pool details
+			List<ShadowPoolInfo> shadowPools = new List<ShadowPoolInfo>();
+			int identicalShadowCount = 0;
+			int differingShadowCount = 0;
+			foreach (PoolId id in shadowIds.OrderBy(x => x.ToString()))
+			{
+				PoolConfig configPool = configPoolLookup[id];
+				IPoolConfig dbPool = dbPoolLookup[id];
+				agentStats.TryGetValue(id, out (int Total, int Online) stats);
+
+				Dictionary<string, FieldDiffEntry>? fieldDiffs = ComparePoolFields(configPool, dbPool);
+				bool fieldsMatch = fieldDiffs == null || fieldDiffs.Count == 0;
+				if (fieldsMatch)
+				{
+					identicalShadowCount++;
+				}
+				else
+				{
+					differingShadowCount++;
+				}
+
+				shadowPools.Add(new ShadowPoolInfo
+				{
+					Id = id.ToString(),
+					Name = configPool.Name,
+					AgentCount = stats.Total,
+					OnlineAgents = stats.Online,
+					Classification = ClassifyPool(stats.Total, stats.Online, configPool),
+					FieldsMatch = fieldsMatch,
+					FieldDiffs = fieldsMatch ? null : fieldDiffs,
+					Action = fieldsMatch
+						? "Remove from database (DB copy is identical)"
+						: "Review differences — config takes precedence at runtime"
+				});
+			}
 
-			// Classify
-			HashSet<PoolId> shadowIds = new HashSet<PoolId>(configPoolIds);
-			shadowIds.IntersectWith(dbPoolIds);
+			// Build DB-only pool details
+			List<DbOnlyPoolInfo> dbOnlyPools = new List<DbOnlyPoolInfo>();
+			foreach (PoolId id in dbOnlyIds.OrderBy(x => x.ToString()))
+			{
+				IPoolConfig dbPool = dbPoolLookup[id];
+				agentStats.TryGetValue(id, out (int Total, int Online) stats);
+				string classification = ClassifyPool(stats.Total, stats.Online, dbPool);
 
-			HashSet<PoolId> dbOnlyIds = new HashSet<PoolId>(dbPoolIds);
-			dbOnlyIds.ExceptWith(configPoolIds);
+				dbOnlyPools.Add(new DbOnlyPoolInfo
+				{
+					Id = id.ToString(),
+					Name = dbPool.Name,
+					AgentCount = stats.Total,
+					OnlineAgents = stats.Online,
+					Classification = classification,
+					Action = classification == "dead"
+						? "Consider removing (no active agents)"
+						: "Add to config files"
+				});
+			}
 
-			return new
+			// Calculate progress
+			int totalPools = configPoolLookup.Count + dbOnlyIds.Count;
+			double migrationPercent = totalPools > 0
+				? Math.Round(100.0 * configPoolLookup.Count / totalPools, 1)
+				: 0;
+
+			// Build recommendations
+			List<string> recommendations = new List<string>();
+			if (identicalShadowCount > 0)
+			{
+				recommendations.Add($"{identicalShadowCount} shadow pools have identical DB entries — safe to delete from database.");
+			}
+			if (dbOnlyIds.Count > 0)
+			{
+				recommendations.Add($"{dbOnlyIds.Count} DB-only pools remain. Use migrate/pool-config to export them.");
+			}
+			if (differingShadowCount > 0)
+			{
+				recommendations.Add($"{differingShadowCount} shadow pools have differing fields — review before deleting DB entries.");
+			}
+
+			return new GetPoolMigrationStatusResponse
 			{
 				Version = (int)computeConfig.VersionEnum,
-				ConfigPools = configPoolIds.Count,
-				DbOnlyPools = dbOnlyIds.Count,
-				ShadowPools = shadowIds.Count,
-				ShadowPoolIds = shadowIds.Select(x => x.ToString()).OrderBy(x => x).ToList(),
-				DbOnlyPoolIds = dbOnlyIds.Select(x => x.ToString()).OrderBy(x => x).ToList()
+				Progress = new PoolMigrationProgressInfo
+				{
+					TotalPools = totalPools,
+					ConfigManaged = configPoolLookup.Count,
+					DbOnly = dbOnlyIds.Count,
+					Shadow = shadowIds.Count,
+					MigrationPercent = migrationPercent
+				},
+				Recommendations = recommendations,
+				ShadowPools = shadowPools,
+				DbOnlyPools = dbOnlyPools
 			};
 		}
+
+		/// <summary>
+		/// Get agent counts per pool using AgentService (cached) or fallback to IAgentCollection.
+		/// </summary>
+		async Task<Dictionary<PoolId, (int Total, int Online)>> GetAgentStatsAsync(CancellationToken cancellationToken)
+		{
+			Dictionary<PoolId, (int Total, int Online)> stats = new Dictionary<PoolId, (int Total, int Online)>();
+
+			AgentService? agentService = _serviceProvider.GetService<AgentService>();
+			if (agentService != null)
+			{
+				IReadOnlyList<IAgent> agents = await agentService.GetCachedAgentsAsync(cancellationToken);
+				foreach (IAgent agent in agents)
+				{
+					bool isOnline = agent.Status == AgentStatus.Ok || agent.Status == AgentStatus.Busy;
+					foreach (PoolId poolId in agent.Pools)
+					{
+						if (!stats.TryGetValue(poolId, out (int Total, int Online) current))
+						{
+							current = (0, 0);
+						}
+						stats[poolId] = (current.Total + 1, current.Online + (isOnline ? 1 : 0));
+					}
+				}
+			}
+			else
+			{
+				// Fallback: use IAgentCollection directly (slower)
+				IAgentCollection? agentCollection = _serviceProvider.GetService<IAgentCollection>();
+				if (agentCollection != null)
+				{
+					await foreach (IAgent agent in agentCollection.FindAsync(cancellationToken: cancellationToken))
+					{
+						bool isOnline = agent.Status == AgentStatus.Ok || agent.Status == AgentStatus.Busy;
+						foreach (PoolId poolId in agent.Pools)
+						{
+							if (!stats.TryGetValue(poolId, out (int Total, int Online) current))
+							{
+								current = (0, 0);
+							}
+							stats[poolId] = (current.Total + 1, current.Online + (isOnline ? 1 : 0));
+						}
+					}
+				}
+			}
+
+			return stats;
+		}
+
+		/// <summary>
+		/// Classify a pool as active, dead, or suspicious based on agent counts.
+		/// </summary>
+		static string ClassifyPool(int totalAgents, int onlineAgents, IPoolConfig poolConfig)
+		{
+			if (totalAgents == 0)
+			{
+				return "dead";
+			}
+			if (onlineAgents == 0)
+			{
+				bool hasFleet = poolConfig.FleetManagers != null && poolConfig.FleetManagers.Count > 0;
+				return hasFleet ? "active" : "suspicious";
+			}
+			return "active";
+		}
+
+		/// <summary>
+		/// Compare config and DB pool fields, returning diffs. Returns null if identical.
+		/// Note: config pools have had MergeDefaults applied (inheritance resolved), while DB pools have not.
+		/// </summary>
+		static Dictionary<string, FieldDiffEntry>? ComparePoolFields(IPoolConfig configPool, IPoolConfig dbPool)
+		{
+			Dictionary<string, FieldDiffEntry> diffs = new Dictionary<string, FieldDiffEntry>();
+
+			CompareField(diffs, "name", configPool.Name, dbPool.Name);
+			CompareField(diffs, "condition", configPool.Condition?.Text, dbPool.Condition?.Text);
+			CompareField(diffs, "enableAutoscaling", configPool.EnableAutoscaling, dbPool.EnableAutoscaling);
+			CompareField(diffs, "minAgents", configPool.MinAgents, dbPool.MinAgents);
+			CompareField(diffs, "numReserveAgents", configPool.NumReserveAgents, dbPool.NumReserveAgents);
+			CompareField(diffs, "conformInterval", configPool.ConformInterval?.ToString(), dbPool.ConformInterval?.ToString());
+			CompareField(diffs, "scaleOutCooldown", configPool.ScaleOutCooldown?.ToString(), dbPool.ScaleOutCooldown?.ToString());
+			CompareField(diffs, "scaleInCooldown", configPool.ScaleInCooldown?.ToString(), dbPool.ScaleInCooldown?.ToString());
+			CompareField(diffs, "shutdownIfDisabledGracePeriod", configPool.ShutdownIfDisabledGracePeriod?.ToString(), dbPool.ShutdownIfDisabledGracePeriod?.ToString());
+			CompareField(diffs, "leaseUtilizationSettings", Serialize(configPool.LeaseUtilizationSettings), Serialize(dbPool.LeaseUtilizationSettings));
+			CompareField(diffs, "jobQueueSettings", Serialize(configPool.JobQueueSettings), Serialize(dbPool.JobQueueSettings));
+			CompareJsonField(diffs, "properties", configPool.Properties, dbPool.Properties);
+			CompareJsonField(diffs, "sizeStrategies", configPool.SizeStrategies, dbPool.SizeStrategies);
+			CompareJsonField(diffs, "fleetManagers", configPool.FleetManagers, dbPool.FleetManagers);
+
+			return diffs.Count > 0 ? diffs : null;
+		}
+
+		static void CompareField(Dictionary<string, FieldDiffEntry> diffs, string fieldName, object? configValue, object? dbValue)
+		{
+			if (!Equals(configValue, dbValue))
+			{
+				diffs[fieldName] = new FieldDiffEntry { Config = configValue, Db = dbValue };
+			}
+		}
+
+		static void CompareJsonField(Dictionary<string, FieldDiffEntry> diffs, string fieldName, object? configValue, object? dbValue)
+		{
+			string configJson = Serialize(configValue);
+			string dbJson = Serialize(dbValue);
+			if (configJson != dbJson)
+			{
+				diffs[fieldName] = new FieldDiffEntry { Config = configValue, Db = dbValue };
+			}
+		}
+
+		static string Serialize(object? value)
+		{
+			if (value == null)
+			{
+				return "null";
+			}
+			return JsonSerializer.Serialize(value, s_compactJsonOptions);
+		}
+
+		/// <summary>
+		/// Diff a shelved P4 changelist's pool configs against the current database.
+		/// Resolves the CL through the full config pipeline (includes, macros, inheritance).
+		/// </summary>
+		[HttpGet]
+		[Route("/api/v1/server/migrate/pool-diff")]
+		public async Task<ActionResult<object>> GetPoolDiffAsync([FromQuery] int cl, [FromQuery] string? cluster = null, CancellationToken cancellationToken = default)
+		{
+			if (!_computeConfig.Value.Authorize(PoolAclAction.ListPools, User))
+			{
+				return Forbid(PoolAclAction.ListPools);
+			}
+
+			// Resolve IChangelistFileReader (lives in Build plugin — may not be registered)
+			IChangelistFileReader? fileReader = _serviceProvider.GetService<IChangelistFileReader>();
+			if (fileReader == null)
+			{
+				return StatusCode(501, new { Error = "Changelist diff requires Perforce integration (HordeServer.Build plugin). This feature is not available in the current deployment." });
+			}
+
+			// Read shelved files
+			Dictionary<Uri, byte[]> files;
+			try
+			{
+				files = await fileReader.ReadChangelistFilesAsync(cl, cluster, cancellationToken);
+			}
+			catch (ChangelistNotFoundException ex)
+			{
+				return NotFound(new { Error = ex.Message });
+			}
+
+			if (files.Count == 0)
+			{
+				return Ok(new
+				{
+					Changelist = cl,
+					Files = Array.Empty<object>(),
+					ValidationError = (string?)null,
+					CurrentConfigPoolCount = _computeConfig.Value.Pools.Count,
+					ProposedConfigPoolCount = _computeConfig.Value.Pools.Count,
+					NewPools = Array.Empty<object>(),
+					DirectlyModifiedPools = Array.Empty<object>(),
+					TemplateAffectedPools = Array.Empty<object>(),
+					MigratingPools = Array.Empty<object>(),
+					RemovedPools = Array.Empty<object>(),
+					RemainingDbPools = Array.Empty<object>(),
+					Summary = new { NewPoolCount = 0, DirectlyModifiedCount = 0, TemplateAffectedCount = 0, MigratingPoolCount = 0, MigratingWithDiffsCount = 0, RemovedPoolCount = 0, RemainingDbPoolCount = 0, CoverageAfterCl = "0%" }
+				});
+			}
+
+			// Resolve proposed config through full pipeline
+			IConfigService configService = _serviceProvider.GetRequiredService<IConfigService>();
+			ConfigResolveResult resolveResult = await configService.ResolveWithOverridesAsync(files, cancellationToken);
+
+			if (resolveResult.Error != null)
+			{
+				return Ok(new
+				{
+					Changelist = cl,
+					Files = files.Keys.Select(u => new { DepotPath = u.AbsolutePath, ConfigUri = u.ToString() }).ToList(),
+					ValidationError = resolveResult.Error
+				});
+			}
+
+			// Extract pools from resolved configs
+			ComputeConfig proposedComputeConfig = resolveResult.Plugins!.GetComputeConfig();
+			Dictionary<PoolId, PoolConfig> proposedPools = proposedComputeConfig.Pools.GroupBy(x => x.Id).ToDictionary(g => g.Key, g => g.First());
+
+			ComputeConfig currentComputeConfig = _computeConfig.Value;
+			Dictionary<PoolId, PoolConfig> currentPools = currentComputeConfig.Pools.GroupBy(x => x.Id).ToDictionary(g => g.Key, g => g.First());
+
+			// Load DB pools
+			IPoolCollection poolCollection = _serviceProvider.GetRequiredService<IPoolCollection>();
+			IReadOnlyList<IPoolConfig> dbPools = await poolCollection.GetDatabasePoolConfigsAsync(cancellationToken);
+			Dictionary<PoolId, IPoolConfig> dbPoolLookup = dbPools.GroupBy(x => x.Id).ToDictionary(g => g.Key, g => g.First());
+
+			// Extract pool IDs from CL files (pre-pipeline) for direct vs inherited change detection
+			HashSet<PoolId> directlyInClPoolIds = ExtractPoolIdsFromRawFiles(files);
+
+			// Diff proposed vs current
+			List<object> newPools = new List<object>();
+			List<object> directlyModifiedPools = new List<object>();
+			List<object> templateAffectedPools = new List<object>();
+			List<object> migratingPools = new List<object>();
+			int migratingWithDiffsCount = 0;
+
+			foreach ((PoolId poolId, PoolConfig proposedPool) in proposedPools)
+			{
+				if (!currentPools.ContainsKey(poolId) && !dbPoolLookup.ContainsKey(poolId))
+				{
+					// New pool
+					newPools.Add(new { Id = poolId.ToString(), Name = proposedPool.Name });
+				}
+				else if (!currentPools.ContainsKey(poolId) && dbPoolLookup.ContainsKey(poolId))
+				{
+					// Migrating pool (DB -> config)
+					IPoolConfig dbPool = dbPoolLookup[poolId];
+					Dictionary<string, FieldDiffEntry>? dbDiffs = ComparePoolFields(proposedPool, dbPool);
+					bool fieldsMatch = dbDiffs == null || dbDiffs.Count == 0;
+					if (!fieldsMatch)
+					{
+						migratingWithDiffsCount++;
+					}
+					migratingPools.Add(new
+					{
+						Id = poolId.ToString(),
+						Name = proposedPool.Name,
+						FieldsMatch = fieldsMatch,
+						DbDiffs = dbDiffs
+					});
+				}
+				else if (currentPools.TryGetValue(poolId, out PoolConfig? currentPool))
+				{
+					// Pool exists in both proposed and current config — check for changes
+					Dictionary<string, FieldDiffEntry>? configDiffs = ComparePoolFields(proposedPool, currentPool);
+					if (configDiffs != null && configDiffs.Count > 0)
+					{
+						if (directlyInClPoolIds.Contains(poolId))
+						{
+							directlyModifiedPools.Add(new
+							{
+								Id = poolId.ToString(),
+								Name = proposedPool.Name,
+								ConfigDiffs = configDiffs
+							});
+						}
+						else
+						{
+							templateAffectedPools.Add(new
+							{
+								Id = poolId.ToString(),
+								Name = proposedPool.Name,
+								TemplateId = proposedPool.Base?.ToString(),
+								ConfigDiffs = configDiffs,
+								Note = proposedPool.Base != null
+									? $"Inherited change from template '{proposedPool.Base}'"
+									: "Inherited change from template"
+							});
+						}
+					}
+				}
+			}
+
+			// Pools removed from config by this CL
+			List<object> removedPools = currentPools.Keys
+				.Where(id => !proposedPools.ContainsKey(id))
+				.OrderBy(id => id.ToString())
+				.Select(id => new { Id = id.ToString(), Name = currentPools[id].Name } as object)
+				.ToList();
+
+			// Remaining DB-only pools
+			List<object> remainingDbPools = dbPoolLookup.Keys
+				.Where(id => !proposedPools.ContainsKey(id))
+				.OrderBy(id => id.ToString())
+				.Select(id => new { Id = id.ToString(), Name = dbPoolLookup[id].Name } as object)
+				.ToList();
+
+			int totalAfterCl = proposedPools.Count + remainingDbPools.Count;
+			string coverage = totalAfterCl > 0
+				? $"{Math.Round(100.0 * proposedPools.Count / totalAfterCl, 1)}%"
+				: "0%";
+
+			return Ok(new
+			{
+				Changelist = cl,
+				Files = files.Keys.Select(u => new { DepotPath = u.AbsolutePath, ConfigUri = u.ToString() }).ToList(),
+				UnreadFiles = resolveResult.UnreadFiles.Select(u => u.ToString()).ToList(),
+				ValidationError = (string?)null,
+				CurrentConfigPoolCount = currentPools.Count,
+				ProposedConfigPoolCount = proposedPools.Count,
+				NewPools = newPools,
+				DirectlyModifiedPools = directlyModifiedPools,
+				TemplateAffectedPools = templateAffectedPools,
+				MigratingPools = migratingPools,
+				RemovedPools = removedPools,
+				RemainingDbPools = remainingDbPools,
+				Summary = new
+				{
+					NewPoolCount = newPools.Count,
+					DirectlyModifiedCount = directlyModifiedPools.Count,
+					TemplateAffectedCount = templateAffectedPools.Count,
+					MigratingPoolCount = migratingPools.Count,
+					MigratingWithDiffsCount = migratingWithDiffsCount,
+					RemovedPoolCount = removedPools.Count,
+					RemainingDbPoolCount = remainingDbPools.Count,
+					CoverageAfterCl = coverage
+				}
+			});
+		}
+
+		/// <summary>
+		/// Extract pool IDs from raw config files (pre-pipeline) for direct vs inherited change detection.
+		/// </summary>
+		static HashSet<PoolId> ExtractPoolIdsFromRawFiles(Dictionary<Uri, byte[]> files)
+		{
+			HashSet<PoolId> poolIds = new HashSet<PoolId>();
+			foreach ((Uri _, byte[] data) in files)
+			{
+				try
+				{
+					using JsonDocument doc = JsonDocument.Parse(data);
+					if (doc.RootElement.TryGetProperty("pools", out JsonElement poolsElement)
+						|| doc.RootElement.TryGetProperty("Pools", out poolsElement))
+					{
+						foreach (JsonElement poolElement in poolsElement.EnumerateArray())
+						{
+							string? id = null;
+							if (poolElement.TryGetProperty("id", out JsonElement idElement)
+								|| poolElement.TryGetProperty("Id", out idElement))
+							{
+								id = idElement.GetString();
+							}
+							else if (poolElement.TryGetProperty("name", out JsonElement nameElement)
+								|| poolElement.TryGetProperty("Name", out nameElement))
+							{
+								// Fallback: derive ID from name using StringId.Sanitize (matches PoolConfig.Id generation)
+								string? name = nameElement.GetString();
+								if (!String.IsNullOrEmpty(name))
+								{
+									id = StringId.Sanitize(name).ToString();
+								}
+							}
+							if (!String.IsNullOrEmpty(id))
+							{
+								poolIds.Add(new PoolId(id));
+							}
+						}
+					}
+				}
+				catch (JsonException)
+				{
+					// Skip files that aren't valid JSON
+				}
+			}
+			return poolIds;
+		}
 	}
 }
