Horde: Fix pool-diff endpoint not matching config files like pools.global.json
The ChangelistFileReader had a hardcoded suffix allowlist that missed .global.json files.
Replace with a simple .json extension check and let the config pipeline filter.

diff --git a/Engine/Source/Programs/Horde/Plugins/Build/HordeServer.Build/VersionControl/Perforce/ChangelistFileReader.cs b/Engine/Source/Programs/Horde/Plugins/Build/HordeServer.Build/VersionControl/Perforce/ChangelistFileReader.cs
--- a/Engine/Source/Programs/Horde/Plugins/Build/HordeServer.Build/VersionControl/Perforce/ChangelistFileReader.cs
+++ b/Engine/Source/Programs/Horde/Plugins/Build/HordeServer.Build/VersionControl/Perforce/ChangelistFileReader.cs
@@ -11,12 +11,6 @@ namespace HordeServer.VersionControl.Perforce
 	/// </summary>
 	public class ChangelistFileReader : IChangelistFileReader
 	{
-		static readonly List<string> s_configSuffixes = new List<string>
-		{
-			"/globals.json", "/global.json", ".project.json", ".stream.json",
-			".dashboard.json", ".telemetry.json"
-		};
-
 		readonly IPerforceService _perforceService;
 		readonly ILogger<ChangelistFileReader> _logger;
 
@@ -50,9 +44,7 @@ namespace HordeServer.VersionControl.Perforce
 
 			foreach (DescribeFileRecord fileRecord in record.Files)
 			{
-				// Filter to config-like files only
-				if (s_configSuffixes.FirstOrDefault(
-					suffix => fileRecord.DepotFile.EndsWith(suffix, StringComparison.OrdinalIgnoreCase)) == null)
+				if (!fileRecord.DepotFile.EndsWith(".json", StringComparison.OrdinalIgnoreCase))
 				{
 					continue;
 				}
