Horde: Updating schema documentation

#rnx

diff --git a/Engine/Source/Programs/Horde/Docs/Config/Schema/AclActions.md b/Engine/Source/Programs/Horde/Docs/Config/Schema/AclActions.md
--- a/Engine/Source/Programs/Horde/Docs/Config/Schema/AclActions.md
+++ b/Engine/Source/Programs/Horde/Docs/Config/Schema/AclActions.md
@@ -37,6 +37,32 @@
 | `DeleteAccount` |  |
 | `ViewAccount` |  |
 
+## Notices
+
+| Name | Description |
+| ---- | ----------- |
+| `CreateNotice` |  |
+| `UpdateNotice` |  |
+| `DeleteNotice` |  |
+
+## ServiceAccounts
+
+| Name | Description |
+| ---- | ----------- |
+| `CreateAccount` |  |
+| `UpdateAccount` |  |
+| `DeleteAccount` |  |
+| `ViewAccount` |  |
+
+## Accounts
+
+| Name | Description |
+| ---- | ----------- |
+| `CreateAccount` |  |
+| `UpdateAccount` |  |
+| `DeleteAccount` |  |
+| `ViewAccount` |  |
+
 ## Artifacts
 
 | Name | Description |
@@ -544,12 +570,6 @@
 | `UpdateReplicator` |  |
 | `ViewReplicator` |  |
 
-## Secrets
-
-| Name | Description |
-| ---- | ----------- |
-| `ViewSecret` | View a credential |
-
 ## ServiceAccounts
 
 | Name | Description |
@@ -721,12 +741,6 @@
 | `UpdateReplicator` |  |
 | `ViewReplicator` |  |
 
-## Secrets
-
-| Name | Description |
-| ---- | ----------- |
-| `ViewSecret` |  |
-
 ## ServiceAccounts
 
 | Name | Description |
@@ -751,182 +765,18 @@
 | `DownloadSoftware` |  |
 | `DeleteSoftware` |  |
 
-## Streams
+## Storage
 
 | Name | Description |
 | ---- | ----------- |
-| `CreateStream` |  |
-| `UpdateStream` |  |
-| `DeleteStream` |  |
-| `ViewStream` |  |
-| `ViewChanges` |  |
-| `ViewTemplate` |  |
-
-## TestData
-
-| Name | Description |
-| ---- | ----------- |
-| `TestHealthRead` |  |
-| `TestHealthWrite` |  |
-
-## Accounts
-
-| Name | Description |
-| ---- | ----------- |
-| `CreateAccount` |  |
-| `UpdateAccount` |  |
-| `DeleteAccount` |  |
-| `ViewAccount` |  |
-
-## Agents
-
-| Name | Description |
-| ---- | ----------- |
-| `CreateAgent` |  |
-| `CreateWorkstationAgent` |  |
-| `UpdateAgent` |  |
-| `DeleteAgent` |  |
-| `ViewAgent` |  |
-| `ListAgents` |  |
-
-## Artifacts
-
-| Name | Description |
-| ---- | ----------- |
-| `ReadArtifact` |  |
-| `WriteArtifact` |  |
-| `DeleteArtifact` |  |
-| `UploadArtifact` |  |
-| `DownloadArtifact` |  |
-
-## Bisect
-
-| Name | Description |
-| ---- | ----------- |
-| `CreateBisectTask` |  |
-| `UpdateBisectTask` |  |
-| `ViewBisectTask` |  |
-
-## BuildHealth
-
-| Name | Description |
-| ---- | ----------- |
-| `AddUpdateDeleteBuildHealthFilter` |  |
-
-## Compute
-
-| Name | Description |
-| ---- | ----------- |
-| `AddComputeTasks` |  |
-| `GetComputeTasks` |  |
-| `UbaCacheRead` |  |
-| `UbaCacheWrite` |  |
-
-## Devices
-
-| Name | Description |
-| ---- | ----------- |
-| `DeviceRead` |  |
-| `DeviceWrite` |  |
-
-## Jobs
-
-| Name | Description |
-| ---- | ----------- |
-| `CreateJob` |  |
-| `UpdateJob` |  |
-| `DeleteJob` |  |
-| `ExecuteJob` |  |
-| `RetryJobStep` |  |
-| `ViewJob` |  |
-
-## Leases
-
-| Name | Description |
-| ---- | ----------- |
-| `ViewLeases` |  |
-| `ViewLeaseTasks` |  |
-
-## Logs
-
-| Name | Description |
-| ---- | ----------- |
-| `CreateLog` |  |
-| `UpdateLog` |  |
-| `ViewLog` |  |
-| `WriteLogData` |  |
-| `CreateEvent` |  |
-| `ViewEvent` |  |
-
-## Notices
-
-| Name | Description |
-| ---- | ----------- |
-| `CreateNotice` |  |
-| `UpdateNotice` |  |
-| `DeleteNotice` |  |
-
-## Notifications
-
-| Name | Description |
-| ---- | ----------- |
-| `CreateSubscription` |  |
-
-## Pools
-
-| Name | Description |
-| ---- | ----------- |
-| `CreatePool` |  |
-| `UpdatePool` |  |
-| `DeletePool` |  |
-| `ViewPool` |  |
-| `ListPools` |  |
-
-## Projects
-
-| Name | Description |
-| ---- | ----------- |
-| `CreateProject` |  |
-| `DeleteProject` |  |
-| `UpdateProject` |  |
-| `ViewProject` |  |
-
-## Replicators
-
-| Name | Description |
-| ---- | ----------- |
-| `UpdateReplicator` |  |
-| `ViewReplicator` |  |
-
-## Secrets
-
-| Name | Description |
-| ---- | ----------- |
-| `ViewSecret` |  |
-
-## ServiceAccounts
-
-| Name | Description |
-| ---- | ----------- |
-| `CreateAccount` |  |
-| `UpdateAccount` |  |
-| `DeleteAccount` |  |
-| `ViewAccount` |  |
-
-## Sessions
-
-| Name | Description |
-| ---- | ----------- |
-| `CreateSession` |  |
-| `ViewSession` |  |
-
-## Software
-
-| Name | Description |
-| ---- | ----------- |
-| `UploadSoftware` |  |
-| `DownloadSoftware` |  |
-| `DeleteSoftware` |  |
+| `ReadBlobs` | Ability to read blobs from the storage service |
+| `WriteBlobs` | Ability to write blobs to the storage service |
+| `ReadAliases` | Ability to read aliases from the storage service |
+| `WriteAliases` | Ability to write aliases to the storage service |
+| `DeleteAliases` | Ability to write aliases to the storage service |
+| `ReadRefs` | Ability to read refs from the storage service |
+| `WriteRefs` | Ability to write refs to the storage service |
+| `DeleteRefs` | Ability to delete refs |
 
 ## Streams
 
@@ -1075,12 +925,6 @@
 | `UpdateReplicator` |  |
 | `ViewReplicator` |  |
 
-## Secrets
-
-| Name | Description |
-| ---- | ----------- |
-| `ViewSecret` |  |
-
 ## ServiceAccounts
 
 | Name | Description |
@@ -1105,6 +949,19 @@
 | `DownloadSoftware` |  |
 | `DeleteSoftware` |  |
 
+## Storage
+
+| Name | Description |
+| ---- | ----------- |
+| `ReadBlobs` |  |
+| `WriteBlobs` |  |
+| `ReadAliases` |  |
+| `WriteAliases` |  |
+| `DeleteAliases` |  |
+| `ReadRefs` |  |
+| `WriteRefs` |  |
+| `DeleteRefs` |  |
+
 ## Streams
 
 | Name | Description |
@@ -1258,12 +1115,6 @@
 | `UpdateReplicator` |  |
 | `ViewReplicator` |  |
 
-## Secrets
-
-| Name | Description |
-| ---- | ----------- |
-| `ViewSecret` |  |
-
 ## ServiceAccounts
 
 | Name | Description |
@@ -1292,14 +1143,14 @@
 
 | Name | Description |
 | ---- | ----------- |
-| `ReadBlobs` | Ability to read blobs from the storage service |
-| `WriteBlobs` | Ability to write blobs to the storage service |
-| `ReadAliases` | Ability to read aliases from the storage service |
-| `WriteAliases` | Ability to write aliases to the storage service |
-| `DeleteAliases` | Ability to write aliases to the storage service |
-| `ReadRefs` | Ability to read refs from the storage service |
-| `WriteRefs` | Ability to write refs to the storage service |
-| `DeleteRefs` | Ability to delete refs |
+| `ReadBlobs` |  |
+| `WriteBlobs` |  |
+| `ReadAliases` |  |
+| `WriteAliases` |  |
+| `DeleteAliases` |  |
+| `ReadRefs` |  |
+| `WriteRefs` |  |
+| `DeleteRefs` |  |
 
 ## Streams
 
@@ -1325,6 +1176,13 @@
 | `TestHealthRead` |  |
 | `TestHealthWrite` |  |
 
+## Tools
+
+| Name | Description |
+| ---- | ----------- |
+| `DownloadTool` | Ability to download a tool |
+| `UploadTool` | Ability to upload new tool versions |
+
 ## Accounts
 
 | Name | Description |
@@ -1458,7 +1316,7 @@
 
 | Name | Description |
 | ---- | ----------- |
-| `ViewSecret` |  |
+| `ViewSecret` | View a credential |
 
 ## ServiceAccounts
 
@@ -1525,8 +1383,8 @@
 
 | Name | Description |
 | ---- | ----------- |
-| `DownloadTool` | Ability to download a tool |
-| `UploadTool` | Ability to upload new tool versions |
+| `DownloadTool` |  |
+| `UploadTool` |  |
 
 ## Accounts
 
diff --git a/Engine/Source/Programs/Horde/Docs/Config/Schema/Globals.md b/Engine/Source/Programs/Horde/Docs/Config/Schema/Globals.md
--- a/Engine/Source/Programs/Horde/Docs/Config/Schema/Globals.md
+++ b/Engine/Source/Programs/Horde/Docs/Config/Schema/Globals.md
@@ -59,15 +59,23 @@
 
 Name | Description
 ---- | -----------
-`build` | [BuildConfig](#buildconfig)<br>Configuration for the build plugin
 `analytics` | [AnalyticsConfig](#analyticsconfig)<br>Configuration for the analytics plugin
-`compute` | [ComputeConfig](#computeconfig)<br>Configuration for the compute plugin
-`secrets` | [SecretsConfig](#secretsconfig)<br>Configuration for the secrets plugin
+`build` | [BuildConfig](#buildconfig)<br>Configuration for the build plugin
 `exampleplugin` | [ExamplePluginConfig](#examplepluginconfig)<br>Configuration for the exampleplugin plugin
+`compute` | [ComputeConfig](#computeconfig)<br>Configuration for the compute plugin
 `experimental` | [ExperimentalConfig](#experimentalconfig)<br>Configuration for the experimental plugin
-`symbols` | [SymbolsConfig](#symbolsconfig)<br>Configuration for the symbols plugin
 `storage` | [StorageConfig](#storageconfig)<br>Configuration for the storage plugin
+`symbols` | [SymbolsConfig](#symbolsconfig)<br>Configuration for the symbols plugin
 `tools` | [ToolsConfig](#toolsconfig)<br>Configuration for the tools plugin
+`secrets` | [SecretsConfig](#secretsconfig)<br>Configuration for the secrets plugin
 
+## AnalyticsConfig
+
+Config settings for analytics
+
+Name | Description
+---- | -----------
+`stores` | [TelemetryStoreConfig](Telemetry.md)`[]`<br>Metrics to aggregate on the Horde server
+
 ## BuildConfig
 
@@ -243,13 +251,13 @@
 `excludeActions` | `string[]`<br>Actions to exclude from the inherited actions
 `extends` | `string[]`<br>Other profiles to extend from
 
-## AnalyticsConfig
+## ExamplePluginConfig
 
-Config settings for analytics
+Configuration for the ExamplePlugin plugin
 
 Name | Description
 ---- | -----------
-`stores` | [TelemetryStoreConfig](Telemetry.md)`[]`<br>Metrics to aggregate on the Horde server
+`greetingMessage` | `string`<br>The Greeting provided by the example plugin
 
 ## ComputeConfig
 
@@ -259,6 +267,7 @@
 ---- | -----------
 `acl` | [AclConfig](#aclconfig)<br>Inherited root acl
 `versionEnum` | [ConfigVersion](#configversion-enum)<br>Config version number
+`telemetryStoreId` | `string`<br>Telemetry store for lease and compute telemetry data
 `rates` | [AgentRateConfig](#agentrateconfig)`[]`<br>List of costs of a particular agent type
 `clusters` | [ComputeClusterConfig](#computeclusterconfig)`[]`<br>List of compute profiles
 `pools` | [PoolConfig](#poolconfig)`[]`<br>List of pools
@@ -447,96 +456,6 @@
 `description` | `string`<br>Human-readable description
 `computeId` | `string`<br>Compute ID for this network (used when allocating compute resources)
 
-## SecretsConfig
-
-Configuration for the secrets system
-
-Name | Description
----- | -----------
-`secrets` | [SecretConfig](#secretconfig)`[]`<br>List of secrets
-`providerConfigs` | [SecretProviderConfig](#secretproviderconfig)`[]`<br>List of configurations for secret providers
-
-## SecretConfig
-
-Configuration for a secret value
-
-Name | Description
----- | -----------
-`id` | `string`<br>Identifier for this secret
-`data` | `string` `->` `string`<br>Key/value pairs associated with this secret
-`sources` | [ExternalSecretConfig](#externalsecretconfig)`[]`<br>Providers to source key/value pairs from
-`acl` | [AclConfig](#aclconfig)<br>Defines access to this particular secret
-
-## ExternalSecretConfig
-
-Configuration for an external secret provider
-
-Name | Description
----- | -----------
-`providerConfig` | `string`<br>The name to look up a secret provider config. A config is mutually exclusive to and to be used when a secret requires additional options than just the name of the provider
-`provider` | `string`<br>Name of the provider to use
-`format` | [ExternalSecretFormat](#externalsecretformat-enum)<br>Format of the secret
-`key` | `string`<br>Optional key indicating the parameter to set in the resulting data array. Required if  is .
-`path` | `string`<br>Optional value indicating what to fetch from the provider
-
-## ExternalSecretFormat (Enum)
-
-Format describing how to parse external secret values
-
-Name | Description
----- | -----------
-`Text` | Secret is a plain text value which will be stored using the external secret key
-`Json` | Secret is a JSON formatted string containing key/value pairs
-
-## SecretProviderConfig
-
-Common settings object for different secret providers
-
-Name | Description
----- | -----------
-`name` | `string`<br>Name of the secret provider config
-`provider` | `string`<br>Name of the provider
-`hcpVault` | [HcpVaultConfig](#hcpvaultconfig)<br>Configuration for HCP Vault
-`azureKeyVault` | [AzureKeyVaultConfig](#azurekeyvaultconfig)<br>Configuration for Azure Key Vault
-
-## HcpVaultConfig
-
-Configuration for HCP Vault secrets.
-
-Name | Description
----- | -----------
-`preSharedKey` | `string`<br>The Vault authentication token to use when the type of credentials is .<br>More information on this token can be found in the API documentation for X-Vault-Token
-`endPoint` | `string`<br>The address of the Vault server.
-`awsArnRole` | `string`<br>The optional AWS ARN to assume a role when logging into Vault. To be used when the type of credentials is . When not set the AWS SDK default credential search will be used.
-`awsIamServerId` | `string`<br>The value for the X-Vault-AWS-IAM-Server-ID header when the Vault server is configured with iam_server_id_header_value. Most likely to be the Vault server's DNS name. To be used when the type of credentials is .
-`role` | `string`<br>The name of the role in the Vault configuration for the login. This role name can be different to the name of the AWS IAM principal. To be used when the type of credentials is . When not set the name of IAM principal will be used.
-`credentials` | [HcpVaultCredentialsType](#hcpvaultcredentialstype-enum)<br>Type of credentials to use.
-
-## HcpVaultCredentialsType (Enum)
-
-Credentials to use to log into HCP Vault.
-
-Name | Description
----- | -----------
-`PreSharedKey` | Use a pre-shared token provided by Vault for authentication.
-`AwsAuth` | Use AWS IAM credentials to authenticate against the Vault server.
-
-## AzureKeyVaultConfig
-
-Configuration for Azure Key Vault secrets.
-
-Name | Description
----- | -----------
-`vaultUri` | `string`<br>The Uri to the Azure Key Vault.
-
-## ExamplePluginConfig
-
-Configuration for the ExamplePlugin plugin
-
-Name | Description
----- | -----------
-`greetingMessage` | `string`<br>The Greeting provided by the example plugin
-
 ## ExperimentalConfig
 
 Configuration for the experimental plugin
@@ -545,25 +464,6 @@
 ---- | -----------
 `notifications` | [NotificationConfig](Notification.md)`[]`<br>Notifications to be sent on the Horde server
 
-## SymbolsConfig
-
-Configuration for the tools system
-
-Name | Description
----- | -----------
-`stores` | [SymbolStoreConfig](#symbolstoreconfig)`[]`<br>List of symbol stores
-
-## SymbolStoreConfig
-
-Configuration for a symbol store
-
-Name | Description
----- | -----------
-`id` | `string`<br>Identifier for this store
-`namespaceId` | `string`<br>Configuration for the symbol store backend
-`public` | `boolean`<br>Whether to make this store available without auth
-`acl` | [AclConfig](#aclconfig)<br>Access to the symbol store
-
 ## StorageConfig
 
 Configuration for storage
@@ -636,6 +536,25 @@
 `enableAliases` | `boolean`<br>Support querying exports by their aliases
 `acl` | [AclConfig](#aclconfig)<br>Access list for this namespace
 
+## SymbolsConfig
+
+Configuration for the tools system
+
+Name | Description
+---- | -----------
+`stores` | [SymbolStoreConfig](#symbolstoreconfig)`[]`<br>List of symbol stores
+
+## SymbolStoreConfig
+
+Configuration for a symbol store
+
+Name | Description
+---- | -----------
+`id` | `string`<br>Identifier for this store
+`namespaceId` | `string`<br>Configuration for the symbol store backend
+`public` | `boolean`<br>Whether to make this store available without auth
+`acl` | [AclConfig](#aclconfig)<br>Access to the symbol store
+
 ## ToolsConfig
 
 Configuration for the tools system
@@ -663,3 +582,85 @@
 `metadata` | `string` `->` `string`<br>Metadata for this tool
 `namespaceId` | `string`<br>Default namespace for new deployments of this tool
 `acl` | [AclConfig](#aclconfig)<br>Permissions for the tool
+
+## SecretsConfig
+
+Configuration for the secrets system
+
+Name | Description
+---- | -----------
+`secrets` | [SecretConfig](#secretconfig)`[]`<br>List of secrets
+`providerConfigs` | [SecretProviderConfig](#secretproviderconfig)`[]`<br>List of configurations for secret providers
+
+## SecretConfig
+
+Configuration for a secret value
+
+Name | Description
+---- | -----------
+`id` | `string`<br>Identifier for this secret
+`data` | `string` `->` `string`<br>Key/value pairs associated with this secret
+`sources` | [ExternalSecretConfig](#externalsecretconfig)`[]`<br>Providers to source key/value pairs from
+`acl` | [AclConfig](#aclconfig)<br>Defines access to this particular secret
+
+## ExternalSecretConfig
+
+Configuration for an external secret provider
+
+Name | Description
+---- | -----------
+`providerConfig` | `string`<br>The name to look up a secret provider config. A config is mutually exclusive to and to be used when a secret requires additional options than just the name of the provider
+`provider` | `string`<br>Name of the provider to use
+`format` | [ExternalSecretFormat](#externalsecretformat-enum)<br>Format of the secret
+`key` | `string`<br>Optional key indicating the parameter to set in the resulting data array. Required if  is .
+`path` | `string`<br>Optional value indicating what to fetch from the provider
+
+## ExternalSecretFormat (Enum)
+
+Format describing how to parse external secret values
+
+Name | Description
+---- | -----------
+`Text` | Secret is a plain text value which will be stored using the external secret key
+`Json` | Secret is a JSON formatted string containing key/value pairs
+
+## SecretProviderConfig
+
+Common settings object for different secret providers
+
+Name | Description
+---- | -----------
+`name` | `string`<br>Name of the secret provider config
+`provider` | `string`<br>Name of the provider
+`hcpVault` | [HcpVaultConfig](#hcpvaultconfig)<br>Configuration for HCP Vault
+`azureKeyVault` | [AzureKeyVaultConfig](#azurekeyvaultconfig)<br>Configuration for Azure Key Vault
+
+## HcpVaultConfig
+
+Configuration for HCP Vault secrets.
+
+Name | Description
+---- | -----------
+`preSharedKey` | `string`<br>The Vault authentication token to use when the type of credentials is .<br>More information on this token can be found in the API documentation for X-Vault-Token
+`endPoint` | `string`<br>The address of the Vault server.
+`awsArnRole` | `string`<br>The optional AWS ARN to assume a role when logging into Vault. To be used when the type of credentials is . When not set the AWS SDK default credential search will be used.
+`awsIamServerId` | `string`<br>The value for the X-Vault-AWS-IAM-Server-ID header when the Vault server is configured with iam_server_id_header_value. Most likely to be the Vault server's DNS name. To be used when the type of credentials is .
+`role` | `string`<br>The name of the role in the Vault configuration for the login. This role name can be different to the name of the AWS IAM principal. To be used when the type of credentials is . When not set the name of IAM principal will be used.
+`credentials` | [HcpVaultCredentialsType](#hcpvaultcredentialstype-enum)<br>Type of credentials to use.
+
+## HcpVaultCredentialsType (Enum)
+
+Credentials to use to log into HCP Vault.
+
+Name | Description
+---- | -----------
+`PreSharedKey` | Use a pre-shared token provided by Vault for authentication.
+`AwsAuth` | Use AWS IAM credentials to authenticate against the Vault server.
+
+## AzureKeyVaultConfig
+
+Configuration for Azure Key Vault secrets.
+
+Name | Description
+---- | -----------
+`vaultUri` | `string`<br>The Uri to the Azure Key Vault.
diff --git a/Engine/Source/Programs/Horde/Docs/Config/Schema/Projects.md b/Engine/Source/Programs/Horde/Docs/Config/Schema/Projects.md
--- a/Engine/Source/Programs/Horde/Docs/Config/Schema/Projects.md
+++ b/Engine/Source/Programs/Horde/Docs/Config/Schema/Projects.md
@@ -18,6 +18,7 @@ Name | Description
 `categories` | [ProjectCategoryConfig](#projectcategoryconfig)`[]`<br>Categories to include in this project
 `jobOptions` | [JobOptions](#joboptions)<br>Default settings for executing jobs
 `workspaceTypes` | `string` `->` [WorkspaceConfig](#workspaceconfig)<br>Default workspace types for streams These are added to the list of each stream's workspace types.
+`ignorePatternFiles` | `string[]`<br>Additional ignore pattern files applied to all streams in this project. Streams can append additional files via their own IgnorePatternFiles property. Paths are relative to workspace root.
 `telemetryStoreId` | `string`<br>Telemetry store for Horde data for this project
 `streams` | [StreamConfig](Streams.md)`[]`<br>List of streams
 `artifactTypes` | [ArtifactTypeConfig](#artifacttypeconfig)`[]`<br>Permissions for artifact types
diff --git a/Engine/Source/Programs/Horde/Docs/Config/Schema/Streams.md b/Engine/Source/Programs/Horde/Docs/Config/Schema/Streams.md
--- a/Engine/Source/Programs/Horde/Docs/Config/Schema/Streams.md
+++ b/Engine/Source/Programs/Horde/Docs/Config/Schema/Streams.md
@@ -15,6 +15,7 @@ Name | Description
 `macros` | [ConfigMacro](#configmacro)`[]`<br>Macros within this stream
 `name` | `string`<br>Name of the stream
 `enginePath` | `string`<br>Path to the engine directory within the workspace. Used for launching UAT.
+`ignorePatternFiles` | `string[]`<br>Additional file paths (relative to workspace root) containing ignore patterns for the log parser. Each file should contain one regex per line, with lines starting with '#' treated as comments. Inherits from ProjectConfig and appends stream-level entries.
 `clusterName` | `string`<br>The perforce cluster containing the stream
 `order` | `integer`<br>Order for this stream
 `initialAgentType` | `string`<br>Default initial agent type for templates
diff --git a/Engine/Source/Programs/Horde/Docs/Deployment/ServerSettings.md b/Engine/Source/Programs/Horde/Docs/Deployment/ServerSettings.md
--- a/Engine/Source/Programs/Horde/Docs/Deployment/ServerSettings.md
+++ b/Engine/Source/Programs/Horde/Docs/Deployment/ServerSettings.md
@@ -127,15 +127,52 @@
 
 Name | Description
 ---- | -----------
-`build` | [BuildServerConfig](#buildserverconfig)<br>Configuration for the build plugin
 `analytics` | [AnalyticsServerConfig](#analyticsserverconfig)<br>Configuration for the analytics plugin
-`compute` | [ComputeServerConfig](#computeserverconfig)<br>Configuration for the compute plugin
-`secrets` | [SecretsServerConfig](#secretsserverconfig)<br>Configuration for the secrets plugin
+`build` | [BuildServerConfig](#buildserverconfig)<br>Configuration for the build plugin
 `exampleplugin` | [ExamplePluginServerConfig](#examplepluginserverconfig)<br>Configuration for the exampleplugin plugin
+`compute` | [ComputeServerConfig](#computeserverconfig)<br>Configuration for the compute plugin
 `experimental` | [ExperimentalServerConfig](#experimentalserverconfig)<br>Configuration for the experimental plugin
-`symbols` | [PluginServerConfig](#pluginserverconfig)<br>Configuration for the symbols plugin
 `storage` | [StorageServerConfig](#storageserverconfig)<br>Configuration for the storage plugin
+`symbols` | [PluginServerConfig](#pluginserverconfig)<br>Configuration for the symbols plugin
 `tools` | [ToolsServerConfig](#toolsserverconfig)<br>Configuration for the tools plugin
+`secrets` | [SecretsServerConfig](#secretsserverconfig)<br>Configuration for the secrets plugin
 
+## AnalyticsServerConfig
+
+Server configuration for the analytics system
+
+Name | Description
+---- | -----------
+`sinks` | [TelemetrySinkConfig](#telemetrysinkconfig)<br>Settings for the various telemetry sinks
+`enabled` | `boolean`<br>Whether the plugin should be enabled or not
+
+## TelemetrySinkConfig
+
+Telemetry sinks
+
+Name | Description
+---- | -----------
+`epic` | [EpicTelemetryConfig](#epictelemetryconfig)<br>Settings for the Epic telemetry sink
+`mongo` | [MongoTelemetryConfig](#mongotelemetryconfig)<br>Settings for the MongoDB telemetry sink
+
+## EpicTelemetryConfig
+
+Configuration for the telemetry sink
+
+Name | Description
+---- | -----------
+`url` | `string`<br>Base URL for the telemetry server
+`appId` | `string`<br>Application name to send in the event messages
+`enabled` | `boolean`<br>Whether to enable this sink
+
+## MongoTelemetryConfig
+
+Configuration for the telemetry sink
+
+Name | Description
+---- | -----------
+`retainDays` | `number`<br>Number of days worth of telmetry events to keep
+`enabled` | `boolean`<br>Whether to enable this sink
+
 ## BuildServerConfig
 
@@ -173,6 +210,7 @@
 `blockCacheDir` | `string`<br>Directory to store the fine-grained block cache. This caches individual exports embedded in bundles.
 `blockCacheSize` | `string`<br>Maximum size of the block cache. Accepts standard binary suffixes. Currently only allocates in multiples of 1024mb.
 `blockCacheSizeBytes` | `integer`<br>Accessor for the block cache size in bytes
+`enableExpandedIgnorePatterns` | `boolean`<br>Whether to enable expanded IgnorePatterns discovery (workspace-wide scanning + config-specified files). When false (default), only Engine directory is scanned for ignore patterns.
 `commits` | [CommitSettings](#commitsettings)<br>Options for the commit service
 `enabled` | `boolean`<br>Whether the plugin should be enabled or not
 
@@ -281,43 +319,14 @@
 `maxChildCount` | `integer`<br>Maximum number of children in each node
 `sliceThreshold` | `integer`<br>Threshold hash value for splitting interior nodes
 
-## AnalyticsServerConfig
+## ExamplePluginServerConfig
 
-Server configuration for the analytics system
+Server configuration for the ExamplePlugin
 
 Name | Description
 ---- | -----------
-`sinks` | [TelemetrySinkConfig](#telemetrysinkconfig)<br>Settings for the various telemetry sinks
 `enabled` | `boolean`<br>Whether the plugin should be enabled or not
 
-## TelemetrySinkConfig
-
-Telemetry sinks
-
-Name | Description
----- | -----------
-`epic` | [EpicTelemetryConfig](#epictelemetryconfig)<br>Settings for the Epic telemetry sink
-`mongo` | [MongoTelemetryConfig](#mongotelemetryconfig)<br>Settings for the MongoDB telemetry sink
-
-## EpicTelemetryConfig
-
-Configuration for the telemetry sink
-
-Name | Description
----- | -----------
-`url` | `string`<br>Base URL for the telemetry server
-`appId` | `string`<br>Application name to send in the event messages
-`enabled` | `boolean`<br>Whether to enable this sink
-
-## MongoTelemetryConfig
-
-Configuration for the telemetry sink
-
-Name | Description
----- | -----------
-`retainDays` | `number`<br>Number of days worth of telmetry events to keep
-`enabled` | `boolean`<br>Whether to enable this sink
-
 ## ComputeServerConfig
 
 Static configuration for the compute plugin
@@ -363,23 +372,6 @@
 `ComputeQueueAwsMetric` | A no-op strategy that reports metrics to let an external AWS auto-scaling policy scale the fleet
 `LeaseUtilizationAwsMetric` | A no-op strategy that reports metrics to let an external AWS auto-scaling policy scale the fleet
 
-## SecretsServerConfig
-
-Static configuration for the secrets plugin
-
-Name | Description
----- | -----------
-`withAws` | `boolean`<br>Whether to enable Amazon Web Services (AWS) specific features
-`enabled` | `boolean`<br>Whether the plugin should be enabled or not
-
-## ExamplePluginServerConfig
-
-Server configuration for the ExamplePlugin
-
-Name | Description
----- | -----------
-`enabled` | `boolean`<br>Whether the plugin should be enabled or not
-
 ## ExperimentalServerConfig
 
 Server configuration for the experimental plugin
@@ -388,14 +380,6 @@
 ---- | -----------
 `enabled` | `boolean`<br>Whether the plugin should be enabled or not
 
-## PluginServerConfig
-
-Base class for plugin server config objects
-
-Name | Description
----- | -----------
-`enabled` | `boolean`<br>Whether the plugin should be enabled or not
-
 ## StorageServerConfig
 
 Static settings for the storage system
@@ -455,6 +439,14 @@
 `AssumeRole` | Assume a particular role. Should specify ARN in
 `AssumeRoleWebIdentity` | Assume a particular role using the current environment variables.
 
+## PluginServerConfig
+
+Base class for plugin server config objects
+
+Name | Description
+---- | -----------
+`enabled` | `boolean`<br>Whether the plugin should be enabled or not
+
 ## ToolsServerConfig
 
 Server configuration for bundled tools
@@ -527,3 +519,12 @@
 `actions` | `string[]`<br>Actions to include
 `excludeActions` | `string[]`<br>Actions to exclude from the inherited actions
 `extends` | `string[]`<br>Other profiles to extend from
+
+## SecretsServerConfig
+
+Static configuration for the secrets plugin
+
+Name | Description
+---- | -----------
+`withAws` | `boolean`<br>Whether to enable Amazon Web Services (AWS) specific features
+`enabled` | `boolean`<br>Whether the plugin should be enabled or not
