Horde: Add TestExecutor case for matching ignore patterns

diff --git a/Engine/Source/Programs/Horde/Drivers/JobDriver/Execution/TestExecutor.cs b/Engine/Source/Programs/Horde/Drivers/JobDriver/Execution/TestExecutor.cs
--- a/Engine/Source/Programs/Horde/Drivers/JobDriver/Execution/TestExecutor.cs
+++ b/Engine/Source/Programs/Horde/Drivers/JobDriver/Execution/TestExecutor.cs
@@ -201,6 +201,7 @@ namespace JobDriver.Execution
 					filter.WriteLine("D:\\Test\\Path\\To\\Source\\File.cpp(234): warning: This is a compilation warning");
 					logger.LogWarning("This is a warning!");
 					filter.WriteLine("warning: this is a test");
+					filter.WriteLine("warning: [IgnorePatternsTest] This warning checks whether IgnorePatterns.txt are enabled and matching");
 				}
 				if (outcome == JobStepOutcome.Failure)
 				{
