Horde: Add temporary Pool Migration dashboard page

Adds /poolmigration page with status display and CL diff tool.
Self-contained in PoolMigration.tsx for easy removal post-migration.

diff --git a/Engine/Source/Programs/Horde/HordeDashboard/src/App.tsx b/Engine/Source/Programs/Horde/HordeDashboard/src/App.tsx
--- a/Engine/Source/Programs/Horde/HordeDashboard/src/App.tsx
+++ b/Engine/Source/Programs/Horde/HordeDashboard/src/App.tsx
@@ -34,6 +34,7 @@ import { AccountsView } from './components/accounts/AccountsView';
 import { HordeLoginView } from './components/accounts/HordeLoginView';
 import { ServiceAccountsView } from './components/accounts/ServiceAccountsView';
 import { AgentRequestsView } from './components/agents/AgentRequestsView';
+import { PoolMigrationView } from './components/agents/PoolMigration';
 import { ArtifactRedirector } from './components/artifacts/ArtifactsRedirector';
 import { DocView } from './components/docs/DocView';
 import { JobDetailViewV2 } from './components/jobDetailsV2/JobDetailViewV2';
@@ -142,6 +143,7 @@ const Main: React.FC = () => {
                { path: "index", element: <UserHomeView /> },
                { path: "project/:projectId", element: <ProjectHome /> },
                { path: "pools", element: <PoolsView /> },
+               { path: "poolmigration", element: <PoolMigrationView /> },
                { path: "job/:jobId", element: <JobDetailViewV2 /> },
                { path: "job", element: <JobRedirector /> },
                { path: "log/:logId", element: <LogView /> },
diff --git a/Engine/Source/Programs/Horde/HordeDashboard/src/components/TopNav.tsx b/Engine/Source/Programs/Horde/HordeDashboard/src/components/TopNav.tsx
--- a/Engine/Source/Programs/Horde/HordeDashboard/src/components/TopNav.tsx
+++ b/Engine/Source/Programs/Horde/HordeDashboard/src/components/TopNav.tsx
@@ -585,6 +585,11 @@ export const TopNav: React.FC<{ suppressServer?: boolean }> = observer(({ suppre
             text: "Pools",
             link: `/pools`
          });
+         resourceItems.push({
+            key: "admin_pool_migration",
+            text: "Pool Migration",
+            link: `/poolmigration`
+         });
       }
 
 
diff --git a/Engine/Source/Programs/Horde/HordeDashboard/src/components/agents/PoolMigration.tsx b/Engine/Source/Programs/Horde/HordeDashboard/src/components/agents/PoolMigration.tsx
new file mode 100644
--- /dev/null
+++ b/Engine/Source/Programs/Horde/HordeDashboard/src/components/agents/PoolMigration.tsx
@@ -0,0 +1,693 @@
+// Copyright Epic Games, Inc. All Rights Reserved.
+//
+// Temporary pool migration dashboard page.
+// Remove this file (and its route in App.tsx / menu item in TopNav.tsx) once migration is complete.
+
+import { DefaultButton, DetailsList, DetailsListLayoutMode, FontIcon, IconButton, IColumn, MessageBar, MessageBarType, Modal, PrimaryButton, SelectionMode, Spinner, SpinnerSize, Stack, Text, TextField } from "@fluentui/react";
+import React, { useEffect, useState } from "react";
+import backend from "../../backend";
+import dashboard, { StatusColor } from "../../backend/Dashboard";
+import { useWindowSize } from "../../base/utilities/hooks";
+import { getHordeStyling } from "../../styles/Styles";
+import { getHordeTheme } from "../../styles/theme";
+import { Breadcrumbs } from "../Breadcrumbs";
+import { TopNav } from "../TopNav";
+import { Markdown } from "../../base/components/Markdown";
+
+// ============================================================================
+// Local types — file-scoped, not exported
+// ============================================================================
+
+interface FieldDiffEntry {
+    config?: unknown;
+    db?: unknown;
+}
+
+interface PoolMigrationProgress {
+    totalPools: number;
+    configManaged: number;
+    dbOnly: number;
+    shadow: number;
+    migrationPercent: number;
+}
+
+interface ShadowPoolInfo {
+    id: string;
+    classification: string;
+    fieldsMatch: boolean;
+    fieldDiffs?: Record<string, FieldDiffEntry>;
+}
+
+interface DbOnlyPoolInfo {
+    id: string;
+    classification: string;
+    agentCount: number;
+}
+
+interface PoolMigrationStatusResponse {
+    version: number;
+    progress: PoolMigrationProgress;
+    recommendations: string[];
+    shadowPools: ShadowPoolInfo[];
+    dbOnlyPools: DbOnlyPoolInfo[];
+}
+
+interface PoolDiffSummary {
+    newPoolCount: number;
+    directlyModifiedCount: number;
+    templateAffectedCount: number;
+    migratingPoolCount: number;
+    migratingWithDiffsCount: number;
+    removedPoolCount: number;
+    remainingDbPoolCount: number;
+    coverageAfterCl: string;
+}
+
+interface PoolDiffResponse {
+    validationError?: string;
+    files: { depotPath: string; configUri: string }[];
+    unreadFiles?: string[];
+    currentConfigPoolCount: number;
+    proposedConfigPoolCount: number;
+    newPools: { id: string }[];
+    directlyModifiedPools: { id: string; configDiffs: Record<string, FieldDiffEntry> }[];
+    templateAffectedPools: { id: string; templateId?: string; configDiffs: Record<string, FieldDiffEntry> }[];
+    migratingPools: { id: string; fieldsMatch: boolean; dbDiffs?: Record<string, FieldDiffEntry> }[];
+    removedPools: { id: string }[];
+    remainingDbPools: { id: string }[];
+    summary: PoolDiffSummary;
+}
+
+// ============================================================================
+// Local fetch helpers
+// ============================================================================
+
+async function localFetch(path: string): Promise<unknown> {
+    const base = backend.serverUrl;
+    const url = base ? `${base}/${path.replace(/^\//, "")}` : `/${path.replace(/^\//, "")}`;
+    const headers: Record<string, string> = {};
+    const token = backend.debugToken;
+    if (token) {
+        headers["Authorization"] = `Bearer ${token}`;
+    }
+    const response = await fetch(url, {
+        method: "GET",
+        mode: "cors",
+        cache: "no-cache",
+        credentials: token ? "same-origin" : "include",
+        redirect: "error",
+        headers
+    });
+    if (!response.ok) {
+        const text = await response.text().catch(() => "");
+        throw new Error(`HTTP ${response.status}: ${text || response.statusText}`);
+    }
+    return response.json();
+}
+
+async function fetchMigrationStatus(): Promise<PoolMigrationStatusResponse> {
+    return localFetch("/api/v1/server/migrate/pool-status") as Promise<PoolMigrationStatusResponse>;
+}
+
+async function fetchPoolDiff(cl: number): Promise<PoolDiffResponse> {
+    return localFetch(`/api/v1/server/migrate/pool-diff?cl=${cl}`) as Promise<PoolDiffResponse>;
+}
+
+// No polling — status is fetched on mount and via manual refresh.
+
+// ============================================================================
+// Embedded README for migrate_pools.py
+// ============================================================================
+
+const MIGRATE_POOLS_README = `# migrate_pools.py — Pool Migration Tool
+
+A standalone Python 3.6+ CLI tool for exporting, diffing, and monitoring pool migration from database to config files.
+
+## Requirements
+
+- Python 3.6+ (stdlib only, no external dependencies)
+- Bearer token with \`ListPools\` permission (set via \`--token\` or the \`$HORDE_TOKEN\` environment variable)
+
+## Usage
+
+\`\`\`
+export HORDE_TOKEN="<your-bearer-token>"
+python3 Scripts/migrate_pools.py --server <url> [options]
+\`\`\`
+
+## Options
+
+| Flag | Required | Description |
+|------|----------|-------------|
+| \`--server\` | Yes | Horde server URL (e.g. \`https://horde.devtools.epicgames.com\`) |
+| \`--token\` | Yes | Bearer token for authentication (or set \`$HORDE_TOKEN\` env var) |
+| \`--output\` | Export mode | Path to write \`pools.global.json\` |
+| \`--diff CL\` | Diff mode | Diff a shelved P4 changelist against current database pools |
+| \`--status\` | Status mode | Show pool migration status |
+| \`--optimize\` | No | Enable inheritance mode — create base pools and reduce redundancy |
+| \`--annotate\` | No | Generate JSONC comments documenting inheritance (requires \`--optimize\`) |
+| \`--dry-run\` | No | Preview classification metrics without writing files |
+| \`--base-names FILE\` | No | Path to a mapping file for custom base pool names |
+
+## Modes
+
+### Export Mode (default)
+Exports all DB-only pools to a JSONC config file. Pools are classified (active, scaled-down, suspicious, dead), sorted by prefix, and written with section headers. Dead pools are commented out.
+
+\`\`\`bash
+python3 Scripts/migrate_pools.py \\
+  --server https://horde.devtools.epicgames.com \\
+  --token $HORDE_TOKEN \\
+  --output /path/to/pools.global.json
+\`\`\`
+
+### Diff Mode
+Compares a shelved P4 changelist against the current database to preview how it would affect pools — showing new, modified, template-affected, migrating, and removed pools.
+
+\`\`\`bash
+python3 Scripts/migrate_pools.py \\
+  --server https://horde.devtools.epicgames.com \\
+  --token $HORDE_TOKEN \\
+  --diff 12345
+\`\`\`
+
+### Status Mode
+Displays migration progress, shadow pool analysis, DB-only pool classification, and recommendations.
+
+\`\`\`bash
+python3 Scripts/migrate_pools.py \\
+  --server https://horde.devtools.epicgames.com \\
+  --token $HORDE_TOKEN \\
+  --status
+\`\`\`
+
+### Optimized Export
+Creates base pools via inheritance clustering to reduce field redundancy across similar pools.
+
+\`\`\`bash
+python3 Scripts/migrate_pools.py \\
+  --server https://horde.devtools.epicgames.com \\
+  --token $HORDE_TOKEN \\
+  --output pools.json --optimize --annotate
+\`\`\`
+`;
+
+// ============================================================================
+// Markdown builders
+// ============================================================================
+
+function buildStatusMarkdown(status: PoolMigrationStatusResponse): string {
+    const p = status.progress;
+    const lines: string[] = [];
+
+    lines.push(`## Migration Progress`);
+    lines.push(``);
+    lines.push(`| Metric | Value |`);
+    lines.push(`|--------|-------|`);
+    lines.push(`| Version | ${status.version} |`);
+    lines.push(`| Total Pools | ${p.totalPools} |`);
+    lines.push(`| Config-Managed | ${p.configManaged} |`);
+    lines.push(`| DB-Only | ${p.dbOnly} |`);
+    lines.push(`| Shadow | ${p.shadow} |`);
+    lines.push(`| Migration | **${p.migrationPercent}%** |`);
+
+    return lines.join("\n");
+}
+
+function formatCell(value: unknown): string {
+    if (value === null || value === undefined) return "*null*";
+    if (typeof value === "object") {
+        // Compact but readable: one line per key for small objects, or wrapped
+        const json = JSON.stringify(value, null, 1);
+        // Collapse to single line if short enough for a table cell
+        const oneLine = json.replace(/\n\s*/g, " ");
+        return `\`${oneLine}\``;
+    }
+    return `\`${String(value)}\``;
+}
+
+function renderFieldDiffs(diffs: Record<string, FieldDiffEntry>, lines: string[]) {
+    const entries = Object.entries(diffs);
+    if (entries.length === 0) return;
+
+    lines.push(``);
+    lines.push(`| Field | DB (current) | Config (proposed) |`);
+    lines.push(`|-------|:-------------|:------------------|`);
+    entries.forEach(([field, d]) => {
+        lines.push(`| \`${field}\` | ${formatCell(d.db)} | ${formatCell(d.config)} |`);
+    });
+}
+
+interface DiffColors {
+    removed: string;
+    added: string;
+    modified: string;
+}
+
+function buildDiffMarkdown(diff: PoolDiffResponse, colors: DiffColors): string {
+    const lines: string[] = [];
+    const s = diff.summary;
+
+    lines.push(`## Summary`);
+    lines.push(``);
+    lines.push(`| Metric | Value |`);
+    lines.push(`|--------|-------|`);
+    lines.push(`| Current config pools | ${diff.currentConfigPoolCount} |`);
+    lines.push(`| Proposed config pools | ${diff.proposedConfigPoolCount} |`);
+    lines.push(`| Coverage after CL | **${s.coverageAfterCl}** |`);
+    lines.push(``);
+
+    if (diff.files.length > 0) {
+        lines.push(`### Files (${diff.files.length})`);
+        diff.files.forEach(f => lines.push(`- \`${f.depotPath}\``));
+        lines.push(``);
+    }
+
+    if (diff.unreadFiles && diff.unreadFiles.length > 0) {
+        lines.push(`### Unread Files (${diff.unreadFiles.length})`);
+        diff.unreadFiles.forEach(f => lines.push(`- \`${f}\``));
+        lines.push(``);
+    }
+
+    // --- Removed pools (first — most critical to notice) ---
+    if (diff.removedPools.length > 0) {
+        lines.push(`### <span style="color:${colors.removed}">Removed Pools (${s.removedPoolCount})</span>`);
+        diff.removedPools.forEach(p => lines.push(`- <span style="color:${colors.removed}">**${p.id}**</span>`));
+        lines.push(``);
+    }
+
+    // --- New pools (second — easy to scan) ---
+    if (diff.newPools.length > 0) {
+        lines.push(`### <span style="color:${colors.added}">New Pools (${s.newPoolCount})</span>`);
+        diff.newPools.forEach(p => lines.push(`- <span style="color:${colors.added}">**${p.id}**</span>`));
+        lines.push(``);
+    }
+
+    // --- Directly modified pools ---
+    if (diff.directlyModifiedPools.length > 0) {
+        lines.push(`### Directly Modified (${s.directlyModifiedCount})`);
+        diff.directlyModifiedPools.forEach(p => {
+            const count = Object.keys(p.configDiffs).length;
+            lines.push(`#### <span style="color:${colors.modified}">${p.id}</span> — ${count} field${count !== 1 ? "s" : ""}`);
+            renderFieldDiffs(p.configDiffs, lines);
+            lines.push(``);
+        });
+    }
+
+    // --- Template-affected pools ---
+    if (diff.templateAffectedPools.length > 0) {
+        lines.push(`### Template-Affected (${s.templateAffectedCount})`);
+        diff.templateAffectedPools.forEach(p => {
+            const count = Object.keys(p.configDiffs).length;
+            const via = p.templateId ? ` via \`${p.templateId}\`` : "";
+            lines.push(`#### <span style="color:${colors.modified}">${p.id}</span>${via} — ${count} field${count !== 1 ? "s" : ""}`);
+            renderFieldDiffs(p.configDiffs, lines);
+            lines.push(``);
+        });
+    }
+
+    // --- Migrating pools ---
+    if (diff.migratingPools.length > 0) {
+        const matching = diff.migratingPools.filter(p => p.fieldsMatch);
+        const differing = diff.migratingPools.filter(p => !p.fieldsMatch);
+
+        if (differing.length > 0) {
+            lines.push(`### Migrating — Config Differs from DB (${s.migratingWithDiffsCount})`);
+            differing.forEach(p => {
+                const diffs = p.dbDiffs ?? {};
+                const count = Object.keys(diffs).length;
+                lines.push(`#### <span style="color:${colors.modified}">${p.id}</span> — ${count} field${count !== 1 ? "s" : ""}`);
+                renderFieldDiffs(diffs, lines);
+                lines.push(``);
+            });
+        }
+        if (matching.length > 0) {
+            lines.push(`### Migrating — Config Matches DB (${matching.length})`);
+            matching.forEach(p => lines.push(`- **${p.id}** (identical)`));
+            lines.push(``);
+        }
+    }
+
+    // --- Remaining DB-only ---
+    if (diff.remainingDbPools.length > 0) {
+        lines.push(`### Remaining DB-Only (${s.remainingDbPoolCount})`);
+        const shown = diff.remainingDbPools.slice(0, 30);
+        shown.forEach(p => lines.push(`- ${p.id}`));
+        if (diff.remainingDbPools.length > 30) {
+            lines.push(`- *... and ${diff.remainingDbPools.length - 30} more*`);
+        }
+        lines.push(``);
+    }
+
+    lines.push(`---`);
+    lines.push(`**Totals:** ${s.newPoolCount} new, ${s.removedPoolCount} removed, ${s.directlyModifiedCount} directly modified, ${s.templateAffectedCount} template-affected, ${s.migratingPoolCount - s.migratingWithDiffsCount} identical migrations, ${s.migratingWithDiffsCount} migrations with diffs`);
+
+    return lines.join("\n");
+}
+
+// ============================================================================
+// Script README Modal
+// ============================================================================
+
+const ScriptReadmeModal: React.FC<{ show: boolean; onClose: () => void }> = ({ show, onClose }) => {
+
+    const { hordeClasses } = getHordeStyling();
+
+    return <Modal
+        className={hordeClasses.modal}
+        isOpen={show}
+        styles={{ main: { padding: 32, width: 900, maxHeight: "85vh" } }}
+        onDismiss={onClose}
+    >
+        <Stack>
+            <Stack horizontal horizontalAlign="space-between" verticalAlign="center" style={{ marginBottom: 16 }}>
+                <Text variant="xLarge" style={{ fontFamily: "Horde Open Sans SemiBold" }}>
+                    migrate_pools.py
+                </Text>
+                <IconButton iconProps={{ iconName: "Cancel" }} onClick={onClose} />
+            </Stack>
+            <div style={{ overflowY: "auto", maxHeight: "calc(85vh - 120px)" }}>
+                <Markdown>{MIGRATE_POOLS_README}</Markdown>
+            </div>
+            <Stack horizontal horizontalAlign="end" style={{ marginTop: 16 }}>
+                <DefaultButton text="Close" onClick={onClose} />
+            </Stack>
+        </Stack>
+    </Modal>;
+};
+
+// ============================================================================
+// MigrationStatusPanel
+// ============================================================================
+
+const MigrationStatusPanel: React.FC<{ status: PoolMigrationStatusResponse | undefined; loading: boolean; onRefresh: () => void }> = ({ status, loading, onRefresh }) => {
+
+    const { modeColors, hordeClasses } = getHordeStyling();
+    const hordeTheme = getHordeTheme();
+    const statusColors = dashboard.getStatusColors();
+    const [showReadme, setShowReadme] = useState(false);
+
+    if (!status) {
+        return <Stack horizontalAlign="center" style={{ paddingTop: 20 }}>
+            <Spinner size={SpinnerSize.large} label="Loading migration status..." />
+        </Stack>;
+    }
+
+    const pct = status.progress.migrationPercent;
+
+    // Shadow pool columns
+    const shadowColumns: IColumn[] = [
+        {
+            key: "id", name: "Pool", fieldName: "id", minWidth: 240, maxWidth: 320, isResizable: true, onRender: (item: ShadowPoolInfo) => {
+                return <Text variant="small">{item.id}</Text>;
+            }
+        },
+        {
+            key: "classification", name: "Classification", minWidth: 100, maxWidth: 120, isResizable: false, onRender: (item: ShadowPoolInfo) => {
+                let color = statusColors.get(StatusColor.Skipped)!;
+                if (item.classification === "active") color = statusColors.get(StatusColor.Success)!;
+                else if (item.classification === "dead") color = statusColors.get(StatusColor.Failure)!;
+                else if (item.classification === "suspicious") color = statusColors.get(StatusColor.Warnings)!;
+                return <Stack horizontal verticalAlign="center" tokens={{ childrenGap: 4 }}>
+                    <FontIcon style={{ color, fontSize: 11 }} iconName="CircleFill" />
+                    <Text variant="small">{item.classification}</Text>
+                </Stack>;
+            }
+        },
+        {
+            key: "fieldsMatch", name: "Fields Match", minWidth: 100, maxWidth: 100, isResizable: false, onRender: (item: ShadowPoolInfo) => {
+                return <Text variant="small" style={{ color: item.fieldsMatch ? statusColors.get(StatusColor.Success)! : statusColors.get(StatusColor.Warnings)! }}>
+                    {item.fieldsMatch ? "Yes" : "No"}
+                </Text>;
+            }
+        },
+        {
+            key: "diffs", name: "Field Diffs", minWidth: 300, isResizable: true, onRender: (item: ShadowPoolInfo) => {
+                if (item.fieldsMatch || !item.fieldDiffs) return null;
+                const entries = Object.entries(item.fieldDiffs);
+                return <Stack tokens={{ childrenGap: 2 }}>
+                    {entries.map(([field, d]) => (
+                        <Text key={field} variant="tiny" style={{ fontFamily: "monospace" }}>
+                            {field}: config={JSON.stringify(d.config)} db={JSON.stringify(d.db)}
+                        </Text>
+                    ))}
+                </Stack>;
+            }
+        }
+    ];
+
+    // DB-only pool columns
+    const dbOnlyColumns: IColumn[] = [
+        {
+            key: "id", name: "Pool", fieldName: "id", minWidth: 240, maxWidth: 320, isResizable: true, onRender: (item: DbOnlyPoolInfo) => {
+                return <Text variant="small">{item.id}</Text>;
+            }
+        },
+        {
+            key: "classification", name: "Classification", minWidth: 100, maxWidth: 120, isResizable: false, onRender: (item: DbOnlyPoolInfo) => {
+                let color = statusColors.get(StatusColor.Skipped)!;
+                if (item.classification === "active") color = statusColors.get(StatusColor.Success)!;
+                else if (item.classification === "dead") color = statusColors.get(StatusColor.Failure)!;
+                else if (item.classification === "suspicious") color = statusColors.get(StatusColor.Warnings)!;
+                return <Stack horizontal verticalAlign="center" tokens={{ childrenGap: 4 }}>
+                    <FontIcon style={{ color, fontSize: 11 }} iconName="CircleFill" />
+                    <Text variant="small">{item.classification}</Text>
+                </Stack>;
+            }
+        },
+        {
+            key: "agentCount", name: "Agents", minWidth: 80, maxWidth: 80, isResizable: false, onRender: (item: DbOnlyPoolInfo) => {
+                return <Text variant="small">{item.agentCount}</Text>;
+            }
+        }
+    ];
+
+    return <Stack tokens={{ childrenGap: 24 }}>
+        {/* Progress card */}
+        <Stack className={hordeClasses.raised}>
+            <Stack styles={{ root: { padding: 16 } }}>
+                <Stack horizontal horizontalAlign="space-between" verticalAlign="center" style={{ marginBottom: 12 }}>
+                    <Text variant="mediumPlus" style={{ fontFamily: "Horde Open Sans SemiBold" }}>Migration Progress</Text>
+                    <IconButton iconProps={{ iconName: "Refresh" }} title="Refresh" onClick={onRefresh} disabled={loading} />
+                </Stack>
+                {/* Progress bar */}
+                <Stack style={{ marginBottom: 12 }}>
+                    <Stack horizontal verticalAlign="center" tokens={{ childrenGap: 8 }}>
+                        <Stack style={{ width: 300, height: 8, backgroundColor: modeColors.crumbs, borderRadius: 4, overflow: "hidden" }}>
+                            <div style={{ width: `${Math.min(pct, 100)}%`, height: "100%", backgroundColor: statusColors.get(StatusColor.Success)!, borderRadius: 4 }} />
+                        </Stack>
+                        <Text variant="small" style={{ fontFamily: "Horde Open Sans SemiBold" }}>{pct}%</Text>
+                    </Stack>
+                </Stack>
+                <Markdown>{buildStatusMarkdown(status)}</Markdown>
+                {status.recommendations.length > 0 && <Stack style={{ marginTop: 12 }} tokens={{ childrenGap: 4 }}>
+                    <Text variant="mediumPlus" style={{ fontFamily: "Horde Open Sans SemiBold" }}>Recommendations</Text>
+                    {status.recommendations.map((r, i) => {
+                        // Rewrite pool-config API references to point to the CLI tool
+                        const rewritten = r
+                            .replace(/\/api\/v1\/server\/migrate\/pool-config/g, "migrate_pools.py --output")
+                            .replace(/pool-config endpoint/gi, "migrate_pools.py export mode")
+                            .replace(/pool-config API/gi, "migrate_pools.py CLI tool");
+                        return <Text key={i} variant="small">- {rewritten}</Text>;
+                    })}
+                    <Text variant="small" style={{ marginTop: 4 }}>
+                        Consider using <Text variant="small" style={{ cursor: "pointer", color: hordeTheme.semanticColors.link, textDecoration: "underline" }} onClick={() => setShowReadme(true)}>migrate_pools.py</Text> to export and optimize pool config migration.
+                    </Text>
+                </Stack>}
+            </Stack>
+        </Stack>
+
+        <ScriptReadmeModal show={showReadme} onClose={() => setShowReadme(false)} />
+
+        {/* Shadow Pools */}
+        {status.shadowPools.length > 0 && <Stack className={hordeClasses.raised}>
+            <Stack styles={{ root: { padding: 16 } }}>
+                <Text variant="mediumPlus" style={{ fontFamily: "Horde Open Sans SemiBold", marginBottom: 8 }}>
+                    Shadow Pools ({status.shadowPools.length})
+                </Text>
+                <DetailsList
+                    items={status.shadowPools}
+                    columns={shadowColumns}
+                    selectionMode={SelectionMode.none}
+                    layoutMode={DetailsListLayoutMode.justified}
+                    compact={true}
+                />
+            </Stack>
+        </Stack>}
+
+        {/* DB-Only Pools */}
+        {status.dbOnlyPools.length > 0 && <Stack className={hordeClasses.raised}>
+            <Stack styles={{ root: { padding: 16 } }}>
+                <Text variant="mediumPlus" style={{ fontFamily: "Horde Open Sans SemiBold", marginBottom: 8 }}>
+                    DB-Only Pools ({status.dbOnlyPools.length})
+                </Text>
+                <DetailsList
+                    items={status.dbOnlyPools}
+                    columns={dbOnlyColumns}
+                    selectionMode={SelectionMode.none}
+                    layoutMode={DetailsListLayoutMode.justified}
+                    compact={true}
+                />
+            </Stack>
+        </Stack>}
+    </Stack>;
+};
+
+// ============================================================================
+// DiffToolPanel + Modal
+// ============================================================================
+
+const DiffToolPanel: React.FC = () => {
+
+    const { hordeClasses } = getHordeStyling();
+
+    const [clInput, setClInput] = useState("");
+    const [loading, setLoading] = useState(false);
+    const [error, setError] = useState<string | undefined>();
+    const [diffResult, setDiffResult] = useState<PoolDiffResponse | undefined>();
+    const [showModal, setShowModal] = useState(false);
+
+    const onDiff = async () => {
+        const cl = parseInt(clInput.trim(), 10);
+        if (isNaN(cl) || cl <= 0) {
+            setError("Please enter a valid changelist number.");
+            return;
+        }
+
+        setError(undefined);
+        setLoading(true);
+        try {
+            const result = await fetchPoolDiff(cl);
+            setDiffResult(result);
+            setShowModal(true);
+        } catch (err) {
+            setError(err instanceof Error ? err.message : String(err));
+        } finally {
+            setLoading(false);
+        }
+    };
+
+    return <Stack tokens={{ childrenGap: 12 }}>
+        <Stack className={hordeClasses.raised}>
+            <Stack styles={{ root: { padding: 16 } }}>
+                <Text variant="mediumPlus" style={{ fontFamily: "Horde Open Sans SemiBold", marginBottom: 12 }}>Pool Config Migration Shelved Changelist</Text>
+                <Stack horizontal verticalAlign="end" tokens={{ childrenGap: 12 }}>
+                    <TextField
+                        label="Changelist Number"
+                        value={clInput}
+                        onChange={(_, v) => setClInput(v ?? "")}
+                        style={{ width: 200 }}
+                        disabled={loading}
+                        onKeyDown={(e) => { if (e.key === "Enter") onDiff(); }}
+                    />
+                    {loading
+                        ? <Spinner size={SpinnerSize.large} label="Fetching diff..." />
+                        : <PrimaryButton text="Diff Changelist" onClick={onDiff} disabled={!clInput.trim()} />
+                    }
+                </Stack>
+                {error && <MessageBar messageBarType={MessageBarType.error} style={{ marginTop: 8 }}>{error}</MessageBar>}
+            </Stack>
+        </Stack>
+
+        {/* Diff Modal */}
+        {diffResult && <PoolDiffModal diff={diffResult} show={showModal} onClose={() => setShowModal(false)} />}
+    </Stack>;
+};
+
+const PoolDiffModal: React.FC<{ diff: PoolDiffResponse; show: boolean; onClose: () => void }> = ({ diff, show, onClose }) => {
+
+    const { hordeClasses } = getHordeStyling();
+    const statusColors = dashboard.getStatusColors();
+    const diffColors: DiffColors = {
+        removed: statusColors.get(StatusColor.Failure)!,
+        added: statusColors.get(StatusColor.Success)!,
+        modified: statusColors.get(StatusColor.Warnings)!
+    };
+
+    const markdownContent = diff.validationError ? undefined : buildDiffMarkdown(diff, diffColors);
+
+    return <Modal
+        className={hordeClasses.modal}
+        isOpen={show}
+        styles={{ main: { padding: 24, width: 1000, maxHeight: "85vh" } }}
+        onDismiss={onClose}
+    >
+        <Stack>
+            <Stack horizontal horizontalAlign="space-between" verticalAlign="center" style={{ marginBottom: 16 }}>
+                <Text variant="xLarge" style={{ fontFamily: "Horde Open Sans SemiBold" }}>
+                    Pool Config Diff
+                </Text>
+                <IconButton
+                    iconProps={{ iconName: "Cancel" }}
+                    onClick={onClose}
+                />
+            </Stack>
+
+            {diff.validationError && <MessageBar messageBarType={MessageBarType.error} style={{ marginBottom: 12 }}>
+                <Text style={{ fontFamily: "Horde Open Sans SemiBold" }}>Config Validation Error</Text>
+                <br />
+                {diff.validationError}
+            </MessageBar>}
+
+            {markdownContent && <div style={{ overflowY: "auto", maxHeight: "calc(85vh - 120px)" }}>
+                <Markdown>{markdownContent}</Markdown>
+            </div>}
+
+            <Stack horizontal horizontalAlign="end" style={{ marginTop: 16 }}>
+                <DefaultButton text="Close" onClick={onClose} />
+            </Stack>
+        </Stack>
+    </Modal>;
+};
+
+// ============================================================================
+// Page view (exported)
+// ============================================================================
+
+export const PoolMigrationView: React.FC = () => {
+
+    const windowSize = useWindowSize();
+    const vw = Math.max(document.documentElement.clientWidth, window.innerWidth || 0);
+    const centerAlign = vw / 2 - 720;
+    const { hordeClasses, modeColors } = getHordeStyling();
+    const key = `windowsize_view_${windowSize.width}_${windowSize.height}`;
+
+    const [status, setStatus] = useState<PoolMigrationStatusResponse | undefined>();
+    const [statusLoading, setStatusLoading] = useState(false);
+
+    const refreshStatus = async () => {
+        setStatusLoading(true);
+        try {
+            const result = await fetchMigrationStatus();
+            setStatus(result);
+        } catch (err) {
+            console.error("Failed to fetch migration status:", err);
+        } finally {
+            setStatusLoading(false);
+        }
+    };
+
+    useEffect(() => {
+        refreshStatus();
+    }, []);
+
+    return <Stack className={hordeClasses.horde}>
+        <TopNav />
+        <Breadcrumbs items={[{ text: "Pool Migration" }]} />
+        <Stack styles={{ root: { width: "100%", backgroundColor: modeColors.background } }}>
+            <Stack style={{ width: "100%", backgroundColor: modeColors.background }}>
+                <Stack style={{ position: "relative", width: "100%", height: 'calc(100vh - 148px)' }}>
+                    <div style={{ overflowX: "auto", overflowY: "visible" }}>
+                        <Stack horizontal style={{ paddingTop: 30, paddingBottom: 48 }}>
+                            <Stack key={key} style={{ paddingLeft: centerAlign }} />
+                            <Stack style={{ width: 1440 }} tokens={{ childrenGap: 24 }}>
+                                <DiffToolPanel />
+                                <MigrationStatusPanel status={status} loading={statusLoading} onRefresh={refreshStatus} />
+                            </Stack>
+                        </Stack>
+                    </div>
+                </Stack>
+            </Stack>
+        </Stack>
+    </Stack>;
+};
