Horde: Update TestExecutor to support loading ignore patterns

diff --git a/Engine/Source/Programs/Horde/Drivers/JobDriver/Execution/JobExecutor.cs b/Engine/Source/Programs/Horde/Drivers/JobDriver/Execution/JobExecutor.cs
--- a/Engine/Source/Programs/Horde/Drivers/JobDriver/Execution/JobExecutor.cs
+++ b/Engine/Source/Programs/Horde/Drivers/JobDriver/Execution/JobExecutor.cs
@@ -1814,7 +1814,7 @@ namespace JobDriver.Execution
 			return result;
 		}
 
-		async Task LoadIgnorePatternsAsync(LogEventParser parser, DirectoryReference workspaceDir, ILogger logger)
+		protected async Task LoadIgnorePatternsAsync(LogEventParser parser, DirectoryReference workspaceDir, ILogger logger)
 		{
 			List<FileReference> patternFiles;
 			if (Batch.EnableExpandedIgnorePatterns)
diff --git a/Engine/Source/Programs/Horde/Drivers/JobDriver/Execution/TestExecutor.cs b/Engine/Source/Programs/Horde/Drivers/JobDriver/Execution/TestExecutor.cs
--- a/Engine/Source/Programs/Horde/Drivers/JobDriver/Execution/TestExecutor.cs
+++ b/Engine/Source/Programs/Horde/Drivers/JobDriver/Execution/TestExecutor.cs
@@ -196,6 +196,8 @@ namespace JobDriver.Execution
 
 			using (LogParser filter = new LogParser(perforceLogger, new List<string>()))
 			{
+				await LoadIgnorePatternsAsync(filter, WorkingDir, logger);
+
 				if (outcome == JobStepOutcome.Warnings)
 				{
 					filter.WriteLine("D:\\Test\\Path\\To\\Source\\File.cpp(234): warning: This is a compilation warning");
