Horde: Add XML doc comment for Job Executor LoadIgnorePatternsAsync
#rnx

diff --git a/Engine/Source/Programs/Horde/Drivers/JobDriver/Execution/JobExecutor.cs b/Engine/Source/Programs/Horde/Drivers/JobDriver/Execution/JobExecutor.cs
--- a/Engine/Source/Programs/Horde/Drivers/JobDriver/Execution/JobExecutor.cs
+++ b/Engine/Source/Programs/Horde/Drivers/JobDriver/Execution/JobExecutor.cs
@@ -1814,6 +1814,9 @@ namespace JobDriver.Execution
 			return result;
 		}
 
+		/// <summary>
+		/// Discovers and loads ignore patterns from the workspace into the given parser
+		/// </summary>
 		protected async Task LoadIgnorePatternsAsync(LogEventParser parser, DirectoryReference workspaceDir, ILogger logger)
 		{
 			List<FileReference> patternFiles;
