Add diagnostic logging for ignore patterns in TestExecutor
#rnx

diff --git a/Engine/Source/Programs/Horde/Drivers/JobDriver/Execution/TestExecutor.cs b/Engine/Source/Programs/Horde/Drivers/JobDriver/Execution/TestExecutor.cs
--- a/Engine/Source/Programs/Horde/Drivers/JobDriver/Execution/TestExecutor.cs
+++ b/Engine/Source/Programs/Horde/Drivers/JobDriver/Execution/TestExecutor.cs
@@ -196,6 +196,16 @@ namespace JobDriver.Execution
 
 			using (LogParser filter = new LogParser(perforceLogger, new List<string>()))
 			{
+				logger.LogInformation("TestExecutor: WorkingDir={WorkingDir}, EnableExpandedIgnorePatterns={Flag}", WorkingDir, Batch.EnableExpandedIgnorePatterns);
+
+				List<FileReference> patternFiles = Batch.EnableExpandedIgnorePatterns
+					? LogEventParser.DiscoverIgnorePatternFiles(WorkingDir, EnginePath)
+					: LogEventParser.DiscoverLegacyIgnorePatternFiles(WorkingDir, EnginePath);
+				foreach (FileReference patternFile in patternFiles)
+				{
+					logger.LogInformation("TestExecutor: Discovered ignore pattern file: {File}", patternFile);
+				}
+
 				await LoadIgnorePatternsAsync(filter, WorkingDir, logger);
 
 				if (outcome == JobStepOutcome.Warnings)
