Horde: Revert TestExecutor ignore pattern diagnostic changes
#rnx

diff --git a/Engine/Source/Programs/Horde/Drivers/JobDriver/Execution/TestExecutor.cs b/Engine/Source/Programs/Horde/Drivers/JobDriver/Execution/TestExecutor.cs
--- a/Engine/Source/Programs/Horde/Drivers/JobDriver/Execution/TestExecutor.cs
+++ b/Engine/Source/Programs/Horde/Drivers/JobDriver/Execution/TestExecutor.cs
@@ -196,24 +196,11 @@ namespace JobDriver.Execution
 
 			using (LogParser filter = new LogParser(perforceLogger, new List<string>()))
 			{
-				logger.LogInformation("TestExecutor: WorkingDir={WorkingDir}, EnableExpandedIgnorePatterns={Flag}", WorkingDir, Batch.EnableExpandedIgnorePatterns);
-
-				List<FileReference> patternFiles = Batch.EnableExpandedIgnorePatterns
-					? LogEventParser.DiscoverIgnorePatternFiles(WorkingDir, EnginePath)
-					: LogEventParser.DiscoverLegacyIgnorePatternFiles(WorkingDir, EnginePath);
-				foreach (FileReference patternFile in patternFiles)
-				{
-					logger.LogInformation("TestExecutor: Discovered ignore pattern file: {File}", patternFile);
-				}
-
-				await LoadIgnorePatternsAsync(filter, WorkingDir, logger);
-
 				if (outcome == JobStepOutcome.Warnings)
 				{
 					filter.WriteLine("D:\\Test\\Path\\To\\Source\\File.cpp(234): warning: This is a compilation warning");
 					logger.LogWarning("This is a warning!");
 					filter.WriteLine("warning: this is a test");
-					filter.WriteLine("warning: [IgnorePatternsTest] This warning checks whether IgnorePatterns.txt are enabled and matching");
 				}
 				if (outcome == JobStepOutcome.Failure)
 				{
