Fix issue with temporal data not being recorded when phase has no dateTime property

#rnx
#Synced-CL 50934132

diff --git a/Engine/Source/Programs/Horde/Plugins/Build/HordeServer.Build/Jobs/TestData/TestDataCollectionV2.cs b/Engine/Source/Programs/Horde/Plugins/Build/HordeServer.Build/Jobs/TestData/TestDataCollectionV2.cs
--- a/Engine/Source/Programs/Horde/Plugins/Build/HordeServer.Build/Jobs/TestData/TestDataCollectionV2.cs
+++ b/Engine/Source/Programs/Horde/Plugins/Build/HordeServer.Build/Jobs/TestData/TestDataCollectionV2.cs
@@ -1058,8 +1058,8 @@ namespace HordeServer.Jobs.TestData
 			{
 				[BsonRequired, BsonElement("testName")]
 				public string TestName { get; set; } = String.Empty;
-				[BsonRequired, BsonElement("dateTime")]
-				public string DateTime { get; set; } = String.Empty;
+				[BsonIgnoreIfNull, BsonElement("dateTime")]
+				public string? DateTime { get; set; } = String.Empty;
 				[BsonRequired, BsonElement("timeElapseSec")]
 				public double TimeElapseSec { get; set; }
 				[BsonIgnoreIfNull, BsonElement("commandline")]
@@ -1080,8 +1080,8 @@ namespace HordeServer.Jobs.TestData
 				public string Name { get; set; } = String.Empty;
 				[BsonRequired, BsonElement("key")]
 				public string Key { get; set; } = String.Empty;
-				[BsonRequired, BsonElement("dateTime")]
-				public string DateTime { get; set; } = String.Empty;
+				[BsonIgnoreIfNull, BsonElement("dateTime")]
+				public string? DateTime { get; set; } = String.Empty;
 				[BsonRequired, BsonElement("timeElapseSec")]
 				public double TimeElapseSec { get; set; }
 				[BsonRequired, BsonElement("outcome")]
@@ -1192,6 +1192,10 @@ namespace HordeServer.Jobs.TestData
 				{
 					sessionDoc.StartDateTime = startDateTime;
 				}
+				else
+				{
+					sessionDoc.StartDateTime = DateTime.UtcNow;
+				}
 				sessionDoc.Duration = TimeSpan.FromSeconds(testSession.Summary.TimeElapseSec);
 				// tags
 				sessionDoc.Tags = tags?.Select(t => t.Id).ToList();
@@ -1221,6 +1225,10 @@ namespace HordeServer.Jobs.TestData
 					{
 						phaseDoc.StartDateTime = startDateTime;
 					}
+					else
+					{
+						phaseDoc.StartDateTime = DateTime.UtcNow;
+					}
 					phaseDoc.Duration = TimeSpan.FromSeconds(phase.TimeElapseSec);
 					// tags
 					phaseDoc.Tags = tags?.Where(tag => phase.Tags?.Contains(tag.Name) ?? false).Select(tag => tag.Id).ToList();
