Move EnableExpandedIgnorePatterns from static BuildServerConfig to dynamic BuildConfig global config

diff --git a/Engine/Source/Programs/Horde/Plugins/Build/HordeServer.Build/BuildConfig.cs b/Engine/Source/Programs/Horde/Plugins/Build/HordeServer.Build/BuildConfig.cs
--- a/Engine/Source/Programs/Horde/Plugins/Build/HordeServer.Build/BuildConfig.cs
+++ b/Engine/Source/Programs/Horde/Plugins/Build/HordeServer.Build/BuildConfig.cs
@@ -189,6 +189,12 @@ namespace HordeServer
 
 		#endregion
 
+		/// <summary>
+		/// Whether to enable expanded IgnorePatterns discovery (workspace-wide scanning + config-specified files).
+		/// When false (default), only Engine directory is scanned for ignore patterns.
+		/// </summary>
+		public bool EnableExpandedIgnorePatterns { get; set; } = true;
+
 		private readonly Dictionary<ProjectId, ProjectConfig> _projectLookup = new Dictionary<ProjectId, ProjectConfig>();
 		private readonly Dictionary<StreamId, StreamConfig> _streamLookup = new Dictionary<StreamId, StreamConfig>();
 
diff --git a/Engine/Source/Programs/Horde/Plugins/Build/HordeServer.Build/BuildServerConfig.cs b/Engine/Source/Programs/Horde/Plugins/Build/HordeServer.Build/BuildServerConfig.cs
--- a/Engine/Source/Programs/Horde/Plugins/Build/HordeServer.Build/BuildServerConfig.cs
+++ b/Engine/Source/Programs/Horde/Plugins/Build/HordeServer.Build/BuildServerConfig.cs
@@ -165,12 +165,6 @@ namespace HordeServer
 		/// </summary>
 		public long BlockCacheSizeBytes => StringUtils.ParseBytesString(BlockCacheSize);
 
-		/// <summary>
-		/// Whether to enable expanded IgnorePatterns discovery (workspace-wide scanning + config-specified files).
-		/// When false (default), only Engine directory is scanned for ignore patterns.
-		/// </summary>
-		public bool EnableExpandedIgnorePatterns { get; set; }
-
 		/// <summary>
 		/// Options for the commit service
 		/// </summary>
diff --git a/Engine/Source/Programs/Horde/Plugins/Build/HordeServer.Build/Jobs/JobRpcService.cs b/Engine/Source/Programs/Horde/Plugins/Build/HordeServer.Build/Jobs/JobRpcService.cs
--- a/Engine/Source/Programs/Horde/Plugins/Build/HordeServer.Build/Jobs/JobRpcService.cs
+++ b/Engine/Source/Programs/Horde/Plugins/Build/HordeServer.Build/Jobs/JobRpcService.cs
@@ -260,7 +260,7 @@ namespace HordeServer.Jobs
 			}
 			response.ValidAgentTypes.Add(streamConfig.AgentTypes.Keys);
 
-			bool expandedIgnorePatterns = _buildServerConfig.Value.EnableExpandedIgnorePatterns;
+			bool expandedIgnorePatterns = _buildConfig.Value.EnableExpandedIgnorePatterns;
 			response.EnableExpandedIgnorePatterns = expandedIgnorePatterns;
 			if (expandedIgnorePatterns && streamConfig.IgnorePatternFiles != null)
 			{
