Shadow pool deletion from dashboard, add configVersion to dashboard config, fix IDE0052 warning

diff --git a/Engine/Source/Programs/Horde/HordeDashboard/src/backend/Api.ts b/Engine/Source/Programs/Horde/HordeDashboard/src/backend/Api.ts
--- a/Engine/Source/Programs/Horde/HordeDashboard/src/backend/Api.ts
+++ b/Engine/Source/Programs/Horde/HordeDashboard/src/backend/Api.ts
@@ -481,6 +481,9 @@ export type GetDashboardConfigResponse = {
 
 	/** Telemetry views */
 	telemetryViews: any[];
+
+	/** The server's global config version number */
+	configVersion?: number;
 }
 
 /**Parameters to register a new agent */
@@ -946,6 +949,9 @@ export type GetPoolResponse = {
 	/** True if pool is managed via config files (read-only in dashboard). Null on pre-V3 servers. */
 	configManaged?: boolean;
 
+	/** True if this config-managed pool also has a database entry (shadow pool) that can be deleted. */
+	hasDatabaseEntry?: boolean;
+
 }
 
 /// Response describing a pool
@@ -966,6 +972,8 @@ export type GetPoolSummaryResponse = {
 	utilization?: number[];
 	/** True if pool is managed via config files (read-only in dashboard). Null on pre-V3 servers. */
 	configManaged?: boolean;
+	/** True if this config-managed pool also has a database entry (shadow pool) that can be deleted. */
+	hasDatabaseEntry?: boolean;
 }
 
 /// Numbers of agents matching various criteria
diff --git a/Engine/Source/Programs/Horde/HordeDashboard/src/backend/Dashboard.ts b/Engine/Source/Programs/Horde/HordeDashboard/src/backend/Dashboard.ts
--- a/Engine/Source/Programs/Horde/HordeDashboard/src/backend/Dashboard.ts
+++ b/Engine/Source/Programs/Horde/HordeDashboard/src/backend/Dashboard.ts
@@ -231,6 +231,10 @@ export class Dashboard {
         return this.config?.deviceProblemCooldownMinutes ?? 30;
     }
 
+    get configVersion(): number | undefined {
+        return this.config?.configVersion;
+    }
+
     get swarmUrl(): string | undefined {
         return this.config?.perforceSwarmUrl;
     }
diff --git a/Engine/Source/Programs/Horde/HordeDashboard/src/components/PoolView.tsx b/Engine/Source/Programs/Horde/HordeDashboard/src/components/PoolView.tsx
--- a/Engine/Source/Programs/Horde/HordeDashboard/src/components/PoolView.tsx
+++ b/Engine/Source/Programs/Horde/HordeDashboard/src/components/PoolView.tsx
@@ -1123,13 +1123,24 @@ const StreamPanel: React.FC = observer(() => {
 
 const PoolPanel: React.FC = () => {
 
-   const [state, setState] = useState<{ showSettings?: boolean }>({})
+   const [state, setState] = useState<{ showSettings?: boolean, showDeleteShadow?: boolean, deletingShadow?: boolean, deleteError?: string }>({})
 
    const pool = handler.pool;
    if (!pool) {
       return null;
    }
 
+   const onDeleteShadow = async () => {
+      setState({ ...state, deletingShadow: true, deleteError: undefined });
+      try {
+         await backend.deletePool(pool.id);
+         handler.pool = await backend.getPool(pool.id);
+         setState({ showDeleteShadow: false, deletingShadow: false });
+      } catch (reason) {
+         setState({ ...state, deletingShadow: false, deleteError: `${reason}` });
+      }
+   };
+
 
    type WorkspaceItem = {
       identifier: string;
@@ -1249,7 +1260,31 @@ const PoolPanel: React.FC = () => {
                      <PrimaryButton text={pool.name} href={`/agents?agent=${encodeURI(pool.id)}&exact=true`} target="_blank" style={{ color: "#FFFFFF", backgroundColor: color, border: "unset", flexShrink: 1 }} />
                      {!!dashboard.user?.dashboardFeatures?.showPoolEditor && pool.configManaged !== true && <IconButton iconProps={{ iconName: "Edit" }} onClick={() => { setState({ ...state, showSettings: true }) }} style={{ color: "#FFFFFF", backgroundColor: color, border: "unset", flexShrink: 1 }} />}
                   </Stack>
-                  {pool.configManaged === true && (
+                  {pool.configManaged === true && pool.hasDatabaseEntry === true && (
+                     <Stack tokens={{ childrenGap: 8 }} styles={{ root: { marginBottom: 8 } }}>
+                        <MessageBar messageBarType={MessageBarType.warning} isMultiline={false}>
+                           This pool is managed via config files but also has a database entry (shadow pool). The database entry can be deleted to complete migration.
+                        </MessageBar>
+                        <DefaultButton text="Delete Database Entry" iconProps={{ iconName: "Delete" }} onClick={() => setState({ ...state, showDeleteShadow: true })} />
+                        {!!state.showDeleteShadow && <Dialog
+                           hidden={false}
+                           onDismiss={() => setState({ ...state, showDeleteShadow: false, deleteError: undefined })}
+                           minWidth={480}
+                           dialogContentProps={{
+                              type: DialogType.largeHeader,
+                              title: `Delete Shadow Pool Entry`,
+                              subText: `This will delete the database entry for pool "${pool.name}". The pool will continue to exist via config files.`
+                           }}
+                           modalProps={{ isBlocking: true }}>
+                           {!!state.deleteError && <MessageBar messageBarType={MessageBarType.error}>{state.deleteError}</MessageBar>}
+                           <DialogFooter>
+                              <PrimaryButton onClick={onDeleteShadow} text={state.deletingShadow ? "Deleting..." : "Delete"} disabled={state.deletingShadow} />
+                              <DefaultButton onClick={() => setState({ ...state, showDeleteShadow: false, deleteError: undefined })} text="Cancel" />
+                           </DialogFooter>
+                        </Dialog>}
+                     </Stack>
+                  )}
+                  {pool.configManaged === true && pool.hasDatabaseEntry !== true && (
                      <MessageBar messageBarType={MessageBarType.info} isMultiline={false} styles={{ root: { marginBottom: 8 } }}>
                         This pool is managed via config files. To modify settings, edit pools.global.json in Perforce.
                      </MessageBar>
diff --git a/Engine/Source/Programs/Horde/HordeDashboard/src/components/PoolsView.tsx b/Engine/Source/Programs/Horde/HordeDashboard/src/components/PoolsView.tsx
--- a/Engine/Source/Programs/Horde/HordeDashboard/src/components/PoolsView.tsx
+++ b/Engine/Source/Programs/Horde/HordeDashboard/src/components/PoolsView.tsx
@@ -172,7 +172,12 @@ const PoolList: React.FC = observer(() => {
                   rootHovered: { border: '0px', backgroundColor: color, color: textColor, },
                   rootPressed: { border: '0px', backgroundColor: color, color: textColor, }
                }} />
-            {pool.configManaged === true && (
+            {pool.configManaged === true && pool.hasDatabaseEntry === true && (
+               <TooltipHost content="Shadow pool - has database entry that can be cleaned up">
+                  <Icon iconName="Warning" styles={{ root: { fontSize: 12, color: statusColors.get(StatusColor.Warnings) }}} />
+               </TooltipHost>
+            )}
+            {pool.configManaged === true && pool.hasDatabaseEntry !== true && (
                <TooltipHost content="Managed via config files (read-only)">
                   <Icon iconName="Lock" styles={{ root: { fontSize: 12, color: statusColors.get(StatusColor.Skipped) }}} />
                </TooltipHost>
diff --git a/Engine/Source/Programs/Horde/HordeDashboard/src/components/TopNav.tsx b/Engine/Source/Programs/Horde/HordeDashboard/src/components/TopNav.tsx
--- a/Engine/Source/Programs/Horde/HordeDashboard/src/components/TopNav.tsx
+++ b/Engine/Source/Programs/Horde/HordeDashboard/src/components/TopNav.tsx
@@ -585,11 +585,13 @@ export const TopNav: React.FC<{ suppressServer?: boolean }> = observer(({ suppre
             text: "Pools",
             link: `/pools`
          });
-         resourceItems.push({
-            key: "admin_pool_migration",
-            text: "Pool Migration",
-            link: `/poolmigration`
-         });
+         if (dashboard.configVersion === 3) {
+            resourceItems.push({
+               key: "admin_pool_migration",
+               text: "Pool Migration",
+               link: `/poolmigration`
+            });
+         }
       }
 
 
diff --git a/Engine/Source/Programs/Horde/HordeServer/Dashboard/DashboardController.cs b/Engine/Source/Programs/Horde/HordeServer/Dashboard/DashboardController.cs
--- a/Engine/Source/Programs/Horde/HordeServer/Dashboard/DashboardController.cs
+++ b/Engine/Source/Programs/Horde/HordeServer/Dashboard/DashboardController.cs
@@ -162,6 +162,8 @@ namespace HordeServer.Dashboard
 
 			dashboardConfigResponse.ArtifactTypeInfo = _globalConfig.Value.Plugins.GetBuildConfig().ArtifactTypes.Select(a => new GetDashboardArtifactTypeResponse { Type = a.Type.ToString(), KeepDays = a.KeepDays}).ToList();
 
+			dashboardConfigResponse.ConfigVersion = _globalConfig.Value.Version;
+
 			foreach (IPluginResponseFilter responseFilter in _responseFilters)
 			{
 				responseFilter.Apply(HttpContext, dashboardConfigResponse);
diff --git a/Engine/Source/Programs/Horde/Plugins/Build/HordeServer.Build/Jobs/JobRpcService.cs b/Engine/Source/Programs/Horde/Plugins/Build/HordeServer.Build/Jobs/JobRpcService.cs
--- a/Engine/Source/Programs/Horde/Plugins/Build/HordeServer.Build/Jobs/JobRpcService.cs
+++ b/Engine/Source/Programs/Horde/Plugins/Build/HordeServer.Build/Jobs/JobRpcService.cs
@@ -59,13 +59,12 @@ namespace HordeServer.Jobs
 		readonly HttpClient _httpClient;
 		readonly IClock _clock;
 		readonly IOptionsSnapshot<BuildConfig> _buildConfig;
-		readonly IOptions<BuildServerConfig> _buildServerConfig;
 		readonly ILogger _logger;
 
 		/// <summary>
 		/// Constructor
 		/// </summary>
-		public JobRpcService(IAclService aclService, JobService jobService, AgentService agentService, PoolService poolService, ConformTaskSource conformTaskSource, IArtifactCollection artifactCollection, IJobCollection jobCollection, ILogCollection logCollection, IGraphCollection graphs, ITestDataCollection testData, ITestDataCollectionV2 testDataV2, IJobStepRefCollection jobStepRefCollection, ITemplateCollection templateCollection, HttpClient httpClient, IClock clock, IOptionsSnapshot<BuildConfig> buildConfig, IOptions<BuildServerConfig> buildServerConfig, ILogger<JobRpcService> logger)
+		public JobRpcService(IAclService aclService, JobService jobService, AgentService agentService, PoolService poolService, ConformTaskSource conformTaskSource, IArtifactCollection artifactCollection, IJobCollection jobCollection, ILogCollection logCollection, IGraphCollection graphs, ITestDataCollection testData, ITestDataCollectionV2 testDataV2, IJobStepRefCollection jobStepRefCollection, ITemplateCollection templateCollection, HttpClient httpClient, IClock clock, IOptionsSnapshot<BuildConfig> buildConfig, ILogger<JobRpcService> logger)
 		{
 			_aclService = aclService;
 			_jobService = jobService;
@@ -83,7 +82,6 @@ namespace HordeServer.Jobs
 			_httpClient = httpClient;
 			_clock = clock;
 			_buildConfig = buildConfig;
-			_buildServerConfig = buildServerConfig;
 			_logger = logger;
 		}
 
diff --git a/Engine/Source/Programs/Horde/Plugins/Compute/HordeServer.Compute/Agents/Pools/PoolMessages.cs b/Engine/Source/Programs/Horde/Plugins/Compute/HordeServer.Compute/Agents/Pools/PoolMessages.cs
--- a/Engine/Source/Programs/Horde/Plugins/Compute/HordeServer.Compute/Agents/Pools/PoolMessages.cs
+++ b/Engine/Source/Programs/Horde/Plugins/Compute/HordeServer.Compute/Agents/Pools/PoolMessages.cs
@@ -539,6 +539,13 @@ namespace HordeServer.Agents.Pools
 		/// </summary>
 		public bool? ConfigManaged { get; set; }
 
+		/// <summary>
+		/// Whether this config-managed pool also has a database entry (shadow pool).
+		/// When true, the database entry can be deleted to complete migration cleanup.
+		/// Only set when <see cref="ConfigManaged"/> is true.
+		/// </summary>
+		public bool? HasDatabaseEntry { get; set; }
+
 		/// <summary>
 		/// Constructor
 		/// </summary>
@@ -580,7 +587,7 @@ namespace HordeServer.Agents.Pools
 	/// <param name="Agents">Truncated list of agents</param>
 	/// <param name="Utilization">Utilization samples for the pool, from zero to one, with one sample for each hour</param>
 	/// <param name="ConfigManaged">Whether this pool is managed by configuration</param>
-	public record class GetPoolSummaryResponse(PoolId Id, string Name, string ColorValue, bool Autoscaled, GetPoolStatsResponse? Stats, List<GetPoolAgentSummaryResponse>? Agents, List<double>? Utilization, bool? ConfigManaged = null);
+	public record class GetPoolSummaryResponse(PoolId Id, string Name, string ColorValue, bool Autoscaled, GetPoolStatsResponse? Stats, List<GetPoolAgentSummaryResponse>? Agents, List<double>? Utilization, bool? ConfigManaged = null, bool? HasDatabaseEntry = null);
 
 	/// <summary>
 	/// Numbers of agents matching various criteria
diff --git a/Engine/Source/Programs/Horde/Plugins/Compute/HordeServer.Compute/Agents/Pools/PoolsController.cs b/Engine/Source/Programs/Horde/Plugins/Compute/HordeServer.Compute/Agents/Pools/PoolsController.cs
--- a/Engine/Source/Programs/Horde/Plugins/Compute/HordeServer.Compute/Agents/Pools/PoolsController.cs
+++ b/Engine/Source/Programs/Horde/Plugins/Compute/HordeServer.Compute/Agents/Pools/PoolsController.cs
@@ -115,13 +115,25 @@ namespace HordeServer.Agents.Pools
 
 			IReadOnlyList<IPoolConfig> poolConfigs = await _poolCollection.GetConfigsAsync(cancellationToken);
 
+			// In V3+, fetch database pool IDs to detect shadow pools
+			HashSet<PoolId>? dbPoolIds = null;
+			if (computeConfig.VersionEnum >= ConfigVersion.IncrementalPoolMigration)
+			{
+				dbPoolIds = new HashSet<PoolId>(await _poolCollection.GetDatabasePoolIdsAsync(cancellationToken));
+			}
+
 			List<object> responses = new List<object>();
 			foreach (IPoolConfig poolConfig in poolConfigs)
 			{
 				GetPoolResponse poolResponse = new GetPoolResponse(poolConfig);
 				if (computeConfig.VersionEnum >= ConfigVersion.IncrementalPoolMigration)
 				{
-					poolResponse.ConfigManaged = computeConfig.TryGetPool(poolConfig.Id, out _);
+					bool isConfigManaged = computeConfig.TryGetPool(poolConfig.Id, out _);
+					poolResponse.ConfigManaged = isConfigManaged;
+					if (isConfigManaged)
+					{
+						poolResponse.HasDatabaseEntry = dbPoolIds!.Contains(poolConfig.Id);
+					}
 				}
 				responses.Add(poolResponse.ApplyFilter(filter));
 			}
@@ -151,6 +163,13 @@ namespace HordeServer.Agents.Pools
 			DateTime utcNow = _clock.UtcNow;
 			IReadOnlyList<IPoolConfig> poolConfigs = await _poolCollection.GetConfigsAsync(cancellationToken);
 
+			// In V3+, fetch database pool IDs to detect shadow pools
+			HashSet<PoolId>? dbPoolIds = null;
+			if (computeConfig.VersionEnum >= ConfigVersion.IncrementalPoolMigration)
+			{
+				dbPoolIds = new HashSet<PoolId>(await _poolCollection.GetDatabasePoolIdsAsync(cancellationToken));
+			}
+
 			// Compute all the expensive response stuff
 			Dictionary<PoolId, PoolInfo>? poolIdToInfo = null;
 			if (stats || numAgents > 0 || numUtilizationSamples > 0)
@@ -241,12 +260,17 @@ namespace HordeServer.Agents.Pools
 
 					// In V3+, indicate whether each pool is managed by config files or the database.
 					bool? configManaged = null;
+					bool? hasDatabaseEntry = null;
 					if (computeConfig.VersionEnum >= ConfigVersion.IncrementalPoolMigration)
 					{
 						configManaged = computeConfig.TryGetPool(poolConfig.Id, out _);
+						if (configManaged == true)
+						{
+							hasDatabaseEntry = dbPoolIds!.Contains(poolConfig.Id);
+						}
 					}
 
-					GetPoolSummaryResponse response = new GetPoolSummaryResponse(poolConfig.Id, poolConfig.Name, poolConfig.GetColorValue(), poolConfig.EnableAutoscaling, statsResponse, agentResponses, utilizationResponses, configManaged);
+					GetPoolSummaryResponse response = new GetPoolSummaryResponse(poolConfig.Id, poolConfig.Name, poolConfig.GetColorValue(), poolConfig.EnableAutoscaling, statsResponse, agentResponses, utilizationResponses, configManaged, hasDatabaseEntry);
 					responses.Add(response);
 				}
 			}
@@ -280,7 +304,13 @@ namespace HordeServer.Agents.Pools
 			GetPoolResponse response = new GetPoolResponse(pool);
 			if (_computeConfig.Value.VersionEnum >= ConfigVersion.IncrementalPoolMigration)
 			{
-				response.ConfigManaged = _computeConfig.Value.TryGetPool(poolIdValue, out _);
+				bool isConfigManaged = _computeConfig.Value.TryGetPool(poolIdValue, out _);
+				response.ConfigManaged = isConfigManaged;
+				if (isConfigManaged)
+				{
+					List<PoolId> dbPoolIds = await _poolCollection.GetDatabasePoolIdsAsync();
+					response.HasDatabaseEntry = dbPoolIds.Contains(poolIdValue);
+				}
 			}
 			return response.ApplyFilter(filter);
 		}
@@ -358,7 +388,12 @@ namespace HordeServer.Agents.Pools
 			if (_computeConfig.Value.VersionEnum >= ConfigVersion.IncrementalPoolMigration
 				&& _computeConfig.Value.TryGetPool(poolIdValue, out _))
 			{
-				return BadRequest($"Pool '{poolId}' is managed via config files and cannot be deleted through the API. Edit pools.global.json in Perforce instead.");
+				// Shadow pool case: pool exists in config — delete only the database entry
+				if (!await _poolCollection.DeleteConfigAsync(poolIdValue))
+				{
+					return BadRequest($"Pool '{poolId}' is managed via config files and has no database entry to remove.");
+				}
+				return Ok();
 			}
 
 			if (!await _poolCollection.DeleteConfigAsync(poolIdValue))
diff --git a/Engine/Source/Programs/Shared/EpicGames.Horde/Dashboard/DashboardMessages.cs b/Engine/Source/Programs/Shared/EpicGames.Horde/Dashboard/DashboardMessages.cs
--- a/Engine/Source/Programs/Shared/EpicGames.Horde/Dashboard/DashboardMessages.cs
+++ b/Engine/Source/Programs/Shared/EpicGames.Horde/Dashboard/DashboardMessages.cs
@@ -80,6 +80,11 @@ namespace EpicGames.Horde.Dashboard
 		/// </summary>
 		public List<GetDashboardArtifactTypeResponse> ArtifactTypeInfo { get; set; } = [];
 
+		/// <summary>
+		/// The server's global config version number. Maps to the ConfigVersion enum on the server.
+		/// </summary>
+		public int? ConfigVersion { get; set; }
+
 	}
 
 	/// <summary>
