Horde: Fix CS1573 warning: add missing XML param tag for HasDatabaseEntry on GetPoolSummaryResponse

diff --git a/Engine/Source/Programs/Horde/Plugins/Compute/HordeServer.Compute/Agents/Pools/PoolMessages.cs b/Engine/Source/Programs/Horde/Plugins/Compute/HordeServer.Compute/Agents/Pools/PoolMessages.cs
--- a/Engine/Source/Programs/Horde/Plugins/Compute/HordeServer.Compute/Agents/Pools/PoolMessages.cs
+++ b/Engine/Source/Programs/Horde/Plugins/Compute/HordeServer.Compute/Agents/Pools/PoolMessages.cs
@@ -587,6 +587,7 @@ namespace HordeServer.Agents.Pools
 	/// <param name="Agents">Truncated list of agents</param>
 	/// <param name="Utilization">Utilization samples for the pool, from zero to one, with one sample for each hour</param>
 	/// <param name="ConfigManaged">Whether this pool is managed by configuration</param>
+	/// <param name="HasDatabaseEntry">Whether this config-managed pool also has a database entry (shadow pool)</param>
 	public record class GetPoolSummaryResponse(PoolId Id, string Name, string ColorValue, bool Autoscaled, GetPoolStatsResponse? Stats, List<GetPoolAgentSummaryResponse>? Agents, List<double>? Utilization, bool? ConfigManaged = null, bool? HasDatabaseEntry = null);
 
 	/// <summary>
