Horde: Add logging for why a schedule job may not be triggered

#rnx

diff --git a/Engine/Source/Programs/Horde/Plugins/Build/HordeServer.Build/Jobs/Schedules/ScheduleService.cs b/Engine/Source/Programs/Horde/Plugins/Build/HordeServer.Build/Jobs/Schedules/ScheduleService.cs
--- a/Engine/Source/Programs/Horde/Plugins/Build/HordeServer.Build/Jobs/Schedules/ScheduleService.cs
+++ b/Engine/Source/Programs/Horde/Plugins/Build/HordeServer.Build/Jobs/Schedules/ScheduleService.cs
@@ -428,6 +428,8 @@ namespace HordeServer.Jobs.Schedules
 			await using IAsyncEnumerator<ICommit> commitEnumerator = commitEnumerable.GetAsyncEnumerator(cancellationToken);
 
 			// Start as many jobs as possible
+			int totalExamined = 0;
+			int totalFiltered = 0;
 			List<(CommitIdWithOrder CommitId, CommitIdWithOrder CodeCommitId)> triggerChanges = new List<(CommitIdWithOrder, CommitIdWithOrder)>();
 			while (triggerChanges.Count < maxNewChanges)
 			{
@@ -471,6 +473,10 @@ namespace HordeServer.Jobs.Schedules
 				}
 
 				// Adjust the changelist for the desired filter
+				if (commit != null)
+				{
+					totalExamined++;
+				}
 				if (commit == null || await ShouldBuildChangeAsync(commit, schedule.Commits, fileFilter, cancellationToken))
 				{
 					CommitIdWithOrder codeCommitId = commitId;
@@ -487,6 +493,10 @@ namespace HordeServer.Jobs.Schedules
 
 					triggerChanges.Add((commitId, codeCommitId));
 				}
+				else
+				{
+					totalFiltered++;
+				}
 
 				// Check we haven't exceeded the time limit
 				if (timer.Elapsed > TimeSpan.FromMinutes(2.0))
@@ -507,6 +517,14 @@ namespace HordeServer.Jobs.Schedules
 			if (triggerChanges.Count == 0)
 			{
 				scheduleLogger.LogInformation("Skipping trigger of {StreamName} template {TemplateId} - no candidate changes after CL {LastTriggerChange}", stream.Id, templateId, schedule.LastTriggerCommitId);
+				if (maxNewChanges <= 0)
+				{
+					scheduleLogger.LogInformation("MaxNewChanges is {MaxNewChanges} for {StreamId} template {TemplateId} - no changes were evaluated", maxNewChanges, stream.Id, templateId);
+				}
+				else if (totalExamined > 0)
+				{
+					scheduleLogger.LogInformation("Examined {TotalExamined} changes after CL {LastTriggerChange}, all {TotalFiltered} were filtered out", totalExamined, schedule.LastTriggerCommitId, totalFiltered);
+				}
 				return;
 			}
 
