Horde: Update testing packages to match Shared, fix trivial analysis warnings

#rnx

#Synced-CL 50981627

diff --git a/Engine/Source/Programs/AutomationTool.Tests/AutomationTool.Tests.csproj b/Engine/Source/Programs/AutomationTool.Tests/AutomationTool.Tests.csproj
--- a/Engine/Source/Programs/AutomationTool.Tests/AutomationTool.Tests.csproj
+++ b/Engine/Source/Programs/AutomationTool.Tests/AutomationTool.Tests.csproj
@@ -12,8 +12,8 @@
 
   <ItemGroup>
     <PackageReference Include="Microsoft.NET.Test.Sdk" Version="17.14.1" />
-    <PackageReference Include="MSTest.TestAdapter" Version="3.6.4" />
-    <PackageReference Include="MSTest.TestFramework" Version="3.6.4" />
+    <PackageReference Include="MSTest.TestAdapter" Version="3.11.1" />
+    <PackageReference Include="MSTest.TestFramework" Version="3.11.1" />
     <PackageReference Include="coverlet.collector" Version="6.0.4">
       <PrivateAssets>all</PrivateAssets>
       <IncludeAssets>runtime; build; native; contentfiles; analyzers; buildtransitive</IncludeAssets>
diff --git a/Engine/Source/Programs/AutomationTool/BuildGraph/Tasks/SymGenTask.cs b/Engine/Source/Programs/AutomationTool/BuildGraph/Tasks/SymGenTask.cs
--- a/Engine/Source/Programs/AutomationTool/BuildGraph/Tasks/SymGenTask.cs
+++ b/Engine/Source/Programs/AutomationTool/BuildGraph/Tasks/SymGenTask.cs
@@ -6,7 +6,6 @@ using System.IO;
 using System.Collections.Concurrent;
 using System.Collections.Generic;
 using System.Linq;
-using System.Text;
 using System.Threading.Tasks;
 using System.Xml;
 using EpicGames.Core;
diff --git a/Engine/Source/Programs/Horde/HordeAgent.Tests/HordeAgent.Tests.csproj b/Engine/Source/Programs/Horde/HordeAgent.Tests/HordeAgent.Tests.csproj
--- a/Engine/Source/Programs/Horde/HordeAgent.Tests/HordeAgent.Tests.csproj
+++ b/Engine/Source/Programs/Horde/HordeAgent.Tests/HordeAgent.Tests.csproj
@@ -22,8 +22,8 @@
     <PackageReference Include="Microsoft.Extensions.Logging.Abstractions" Version="8.0.3" />
     <PackageReference Include="Microsoft.NET.Test.Sdk" Version="17.14.1" />
     <PackageReference Include="Moq" Version="4.20.72" />
-    <PackageReference Include="MSTest.TestAdapter" Version="3.6.4" />
-    <PackageReference Include="MSTest.TestFramework" Version="3.6.4" />
+    <PackageReference Include="MSTest.TestAdapter" Version="3.11.1" />
+    <PackageReference Include="MSTest.TestFramework" Version="3.11.1" />
     <PackageReference Include="coverlet.collector" Version="6.0.4">
       <PrivateAssets>all</PrivateAssets>
       <IncludeAssets>runtime; build; native; contentfiles; analyzers; buildtransitive</IncludeAssets>
diff --git a/Engine/Source/Programs/Horde/HordeAgent.Tests/Utility/WindowsProcessMemoryTrackerTests.cs b/Engine/Source/Programs/Horde/HordeAgent.Tests/Utility/WindowsProcessMemoryTrackerTests.cs
--- a/Engine/Source/Programs/Horde/HordeAgent.Tests/Utility/WindowsProcessMemoryTrackerTests.cs
+++ b/Engine/Source/Programs/Horde/HordeAgent.Tests/Utility/WindowsProcessMemoryTrackerTests.cs
@@ -314,7 +314,7 @@ public class WindowsProcessMemoryTrackerTests
 
 		tracker.Start();
 
-		Assert.ThrowsException<InvalidOperationException>(() => tracker.Start());
+		Assert.ThrowsExactly<InvalidOperationException>(() => tracker.Start());
 	}
 
 	[TestMethod]
@@ -328,7 +328,7 @@ public class WindowsProcessMemoryTrackerTests
 		}
 
 		// Constructor throws before any resources are allocated, so no dispose needed
-		Assert.ThrowsException<ArgumentNullException>(() =>
+		Assert.ThrowsExactly<ArgumentNullException>(() =>
 			_ = new WindowsProcessMemoryTracker(null!, TimeSpan.FromSeconds(30), _logger));
 	}
 
@@ -345,7 +345,7 @@ public class WindowsProcessMemoryTrackerTests
 		using ManagedProcessGroup processGroup = new();
 
 		// Constructor throws before any resources are allocated, so no dispose needed
-		Assert.ThrowsException<ArgumentNullException>(() =>
+		Assert.ThrowsExactly<ArgumentNullException>(() =>
 			_ = new WindowsProcessMemoryTracker(processGroup, TimeSpan.FromSeconds(30), null!));
 	}
 }
diff --git a/Engine/Source/Programs/Horde/HordeServer.Tests.Shared/DatabaseIntegrationTest.cs b/Engine/Source/Programs/Horde/HordeServer.Tests.Shared/DatabaseIntegrationTest.cs
--- a/Engine/Source/Programs/Horde/HordeServer.Tests.Shared/DatabaseIntegrationTest.cs
+++ b/Engine/Source/Programs/Horde/HordeServer.Tests.Shared/DatabaseIntegrationTest.cs
@@ -10,7 +10,6 @@ using Microsoft.Extensions.DependencyInjection;
 using Microsoft.Extensions.Logging;
 using Microsoft.Extensions.Logging.Abstractions;
 using Microsoft.Extensions.Options;
-using Microsoft.VisualStudio.TestTools.UnitTesting;
 using MongoDB.Driver;
 using StackExchange.Redis;
 
diff --git a/Engine/Source/Programs/Horde/HordeServer.Tests.Shared/HordeServer.Tests.Shared.csproj b/Engine/Source/Programs/Horde/HordeServer.Tests.Shared/HordeServer.Tests.Shared.csproj
--- a/Engine/Source/Programs/Horde/HordeServer.Tests.Shared/HordeServer.Tests.Shared.csproj
+++ b/Engine/Source/Programs/Horde/HordeServer.Tests.Shared/HordeServer.Tests.Shared.csproj
@@ -26,8 +26,8 @@
     <PackageReference Include="Microsoft.AspNetCore.Mvc.Testing" Version="8.0.20" />
     <PackageReference Include="Microsoft.NET.Test.Sdk" Version="17.14.1" />
     <PackageReference Include="Moq" Version="4.20.72" />
-    <PackageReference Include="MSTest.TestAdapter" Version="3.6.4" />
-    <PackageReference Include="MSTest.TestFramework" Version="3.6.4" />
+    <PackageReference Include="MSTest.TestAdapter" Version="3.11.1" />
+    <PackageReference Include="MSTest.TestFramework" Version="3.11.1" />
     <PackageReference Include="Microsoft.Extensions.Caching.Memory" Version="8.0.1" />
     <PackageReference Include="coverlet.collector" Version="6.0.4">
       <PrivateAssets>all</PrivateAssets>
diff --git a/Engine/Source/Programs/Horde/HordeServer.Tests/Accounts/AccountCollectionTests.cs b/Engine/Source/Programs/Horde/HordeServer.Tests/Accounts/AccountCollectionTests.cs
--- a/Engine/Source/Programs/Horde/HordeServer.Tests/Accounts/AccountCollectionTests.cs
+++ b/Engine/Source/Programs/Horde/HordeServer.Tests/Accounts/AccountCollectionTests.cs
@@ -42,8 +42,8 @@ namespace HordeServer.Tests.Accounts
 			Assert.IsNotNull(account2);
 			Assert.AreEqual(account.Id, account2.Id);
 
-			await Assert.ThrowsExceptionAsync<LoginAlreadyTakenException>(() => accountCollection.CreateAsync(new CreateAccountOptions("Bill", "bob")));
-			await Assert.ThrowsExceptionAsync<LoginAlreadyTakenException>(() => account2.UpdateAsync(new UpdateAccountOptions { Login = "admin" }));
+			await Assert.ThrowsExactlyAsync<LoginAlreadyTakenException>(() => accountCollection.CreateAsync(new CreateAccountOptions("Bill", "bob")));
+			await Assert.ThrowsExactlyAsync<LoginAlreadyTakenException>(() => account2.UpdateAsync(new UpdateAccountOptions { Login = "admin" }));
 		}
 	}
 }
diff --git a/Engine/Source/Programs/Horde/HordeServer.Tests/HordeServer.Tests.csproj b/Engine/Source/Programs/Horde/HordeServer.Tests/HordeServer.Tests.csproj
--- a/Engine/Source/Programs/Horde/HordeServer.Tests/HordeServer.Tests.csproj
+++ b/Engine/Source/Programs/Horde/HordeServer.Tests/HordeServer.Tests.csproj
@@ -23,8 +23,8 @@
     <PackageReference Include="Microsoft.AspNetCore.Mvc.Testing" Version="8.0.20" />
     <PackageReference Include="Microsoft.NET.Test.Sdk" Version="17.14.1" />
     <PackageReference Include="Moq" Version="4.20.72" />
-    <PackageReference Include="MSTest.TestAdapter" Version="3.6.4" />
-    <PackageReference Include="MSTest.TestFramework" Version="3.6.4" />
+    <PackageReference Include="MSTest.TestAdapter" Version="3.11.1" />
+    <PackageReference Include="MSTest.TestFramework" Version="3.11.1" />
     <PackageReference Include="Microsoft.Extensions.Caching.Memory" Version="8.0.1" />
     <PackageReference Include="coverlet.collector" Version="6.0.4">
       <PrivateAssets>all</PrivateAssets>
diff --git a/Engine/Source/Programs/Horde/HordeServer.Tests/Plugins/PluginTests.cs b/Engine/Source/Programs/Horde/HordeServer.Tests/Plugins/PluginTests.cs
--- a/Engine/Source/Programs/Horde/HordeServer.Tests/Plugins/PluginTests.cs
+++ b/Engine/Source/Programs/Horde/HordeServer.Tests/Plugins/PluginTests.cs
@@ -114,13 +114,12 @@ namespace HordeServer.Tests.Plugins
 		}
 
 		[TestMethod]
-		[ExpectedException(typeof(InvalidOperationException))]
 		public void TopologicalSortErrorWithCycle()
 		{
 			PluginCollection pluginCollection = new();
 			pluginCollection.Add(typeof(PluginsWithCycle.PluginA));
 			pluginCollection.Add(typeof(PluginsWithCycle.PluginB));
-			_ = PluginCollection.GetTopologicalSort(pluginCollection.LoadedPlugins);
+			Assert.ThrowsExactly<InvalidOperationException>(() => _ = PluginCollection.GetTopologicalSort(pluginCollection.LoadedPlugins));
 		}
 
 		[TestMethod]
diff --git a/Engine/Source/Programs/Horde/HordeServer.Tests/Server/DowntimeServiceTests.cs b/Engine/Source/Programs/Horde/HordeServer.Tests/Server/DowntimeServiceTests.cs
--- a/Engine/Source/Programs/Horde/HordeServer.Tests/Server/DowntimeServiceTests.cs
+++ b/Engine/Source/Programs/Horde/HordeServer.Tests/Server/DowntimeServiceTests.cs
@@ -183,15 +183,14 @@ namespace HordeServer.Tests.Server
 		}
 		
 		[TestMethod]
-		[ExpectedException(typeof(FormatException))]
 		public void WrongDayOfWeekNameError()
 		{
-			GetDowntimeFromJson("""
+			Assert.ThrowsExactly<FormatException>(() => GetDowntimeFromJson("""
 				{
 					"startTime": "Mon Sep 6, 2020, 3AM",
 					"finishTime": "Sun Sep 6, 2020, 3:30AM",
 				}
-			""");
+			"""));
 		}
 
 		[TestMethod]
@@ -277,15 +276,14 @@ namespace HordeServer.Tests.Server
 		}
 
 		[TestMethod]
-		[ExpectedException(typeof(TimeZoneNotFoundException))]
 		public void WrongTimeZoneName()
 		{
-			GetDowntimeFromJson("""
+			Assert.ThrowsExactly<TimeZoneNotFoundException>(() => GetDowntimeFromJson("""
 				{
 					"startTime": "Jan 1, 2020, 3AM",
 					"timezone": "ET"
 				}
-			""");
+			"""));
 		}
 
 		[TestMethod]
diff --git a/Engine/Source/Programs/Horde/HordeServer.Tests/Server/MongoMigrationTests.cs b/Engine/Source/Programs/Horde/HordeServer.Tests/Server/MongoMigrationTests.cs
--- a/Engine/Source/Programs/Horde/HordeServer.Tests/Server/MongoMigrationTests.cs
+++ b/Engine/Source/Programs/Horde/HordeServer.Tests/Server/MongoMigrationTests.cs
@@ -107,7 +107,7 @@ public class MongoMigratorTests : DatabaseIntegrationTest
 	[TestMethod]
 	public void AddMigration_WithMissingAttribute()
 	{
-		MigrationException me = Assert.ThrowsException<MigrationException>(() => _migrator.AddMigration(new MissingAttributeMigration()));
+		MigrationException me = Assert.ThrowsExactly<MigrationException>(() => _migrator.AddMigration(new MissingAttributeMigration()));
 		Assert.IsTrue(me.Message.Contains("is missing [MongoMigration] attribute", StringComparison.Ordinal));
 	}
 	
@@ -115,7 +115,7 @@ public class MongoMigratorTests : DatabaseIntegrationTest
 	public void AddMigration_WithDuplicateVersion()
 	{
 		_migrator.AddMigration(new DuplicateVersionMigrationA());
-		MigrationException me = Assert.ThrowsException<MigrationException>(() => _migrator.AddMigration(new DuplicateVersionMigrationB()));
+		MigrationException me = Assert.ThrowsExactly<MigrationException>(() => _migrator.AddMigration(new DuplicateVersionMigrationB()));
 		Assert.IsTrue(me.Message.Contains("already contains a migration with version 1", StringComparison.Ordinal));
 	}
 	
@@ -148,21 +148,21 @@ public class MongoMigratorTests : DatabaseIntegrationTest
 	{
 		_migrator.AddMigration(new FooMigration1());
 		_migrator.AddMigration(new FooMigration3());
-		MigrationException me = await Assert.ThrowsExceptionAsync<MigrationException>(() => _migrator.UpgradeAsync("foo", null, CancellationToken.None));
+		MigrationException me = await Assert.ThrowsExactlyAsync<MigrationException>(() => _migrator.UpgradeAsync("foo", null, CancellationToken.None));
 		Assert.IsTrue(me.Message.StartsWith("Gap detected", StringComparison.Ordinal));
 	}
 	
 	[TestMethod]
 	public void AddMigration_WithInvalidVersionNumber()
 	{
-		MigrationException me = Assert.ThrowsException<MigrationException>(() => _migrator.AddMigration(new InvalidVersionMigration()));
+		MigrationException me = Assert.ThrowsExactly<MigrationException>(() => _migrator.AddMigration(new InvalidVersionMigration()));
 		Assert.IsTrue(me.Message.Contains("invalid version number", StringComparison.Ordinal));
 	}
 	
 	[TestMethod]
 	public void AddMigration_WithInvalidName()
 	{
-		MigrationException me = Assert.ThrowsException<MigrationException>(() => _migrator.AddMigration(new ShortNameMigration()));
+		MigrationException me = Assert.ThrowsExactly<MigrationException>(() => _migrator.AddMigration(new ShortNameMigration()));
 		Assert.IsTrue(me.Message.Contains("too short/long name", StringComparison.Ordinal));
 	}
 	
@@ -265,7 +265,7 @@ public class MongoMigratorTests : DatabaseIntegrationTest
 	[TestMethod]
 	public async Task CheckSchemaVersion_WithUnknownPlugin_Async()
 	{
-		MigrationException me = await Assert.ThrowsExceptionAsync<MigrationException>(() => _migrator.GetSchemaStateAsync("unknown-plugin", CancellationToken.None));
+		MigrationException me = await Assert.ThrowsExactlyAsync<MigrationException>(() => _migrator.GetSchemaStateAsync("unknown-plugin", CancellationToken.None));
 		Assert.AreEqual("Unknown plugin unknown-plugin", me.Message);
 	}
 	
diff --git a/Engine/Source/Programs/Horde/Plugins/Analytics/HordeServer.Analytics.Integration.Tests/HordeServer.Analytics.Integration.Tests.csproj b/Engine/Source/Programs/Horde/Plugins/Analytics/HordeServer.Analytics.Integration.Tests/HordeServer.Analytics.Integration.Tests.csproj
--- a/Engine/Source/Programs/Horde/Plugins/Analytics/HordeServer.Analytics.Integration.Tests/HordeServer.Analytics.Integration.Tests.csproj
+++ b/Engine/Source/Programs/Horde/Plugins/Analytics/HordeServer.Analytics.Integration.Tests/HordeServer.Analytics.Integration.Tests.csproj
@@ -25,8 +25,8 @@
       <IncludeAssets>runtime; build; native; contentfiles; analyzers; buildtransitive</IncludeAssets>
     </PackageReference>
     <PackageReference Include="Microsoft.NET.Test.Sdk" Version="17.14.1" />
-    <PackageReference Include="MSTest.TestAdapter" Version="3.6.4" />
-    <PackageReference Include="MSTest.TestFramework" Version="3.6.4" />
+    <PackageReference Include="MSTest.TestAdapter" Version="3.11.1" />
+    <PackageReference Include="MSTest.TestFramework" Version="3.11.1" />
   </ItemGroup>
 
   <ItemGroup>
diff --git a/Engine/Source/Programs/Horde/Plugins/Analytics/HordeServer.Analytics.Tests/HordeServer.Analytics.Tests.csproj b/Engine/Source/Programs/Horde/Plugins/Analytics/HordeServer.Analytics.Tests/HordeServer.Analytics.Tests.csproj
--- a/Engine/Source/Programs/Horde/Plugins/Analytics/HordeServer.Analytics.Tests/HordeServer.Analytics.Tests.csproj
+++ b/Engine/Source/Programs/Horde/Plugins/Analytics/HordeServer.Analytics.Tests/HordeServer.Analytics.Tests.csproj
@@ -25,8 +25,8 @@
       <IncludeAssets>runtime; build; native; contentfiles; analyzers; buildtransitive</IncludeAssets>
     </PackageReference>
     <PackageReference Include="Microsoft.NET.Test.Sdk" Version="17.14.1" />
-    <PackageReference Include="MSTest.TestAdapter" Version="3.6.4" />
-    <PackageReference Include="MSTest.TestFramework" Version="3.6.4" />
+    <PackageReference Include="MSTest.TestAdapter" Version="3.11.1" />
+    <PackageReference Include="MSTest.TestFramework" Version="3.11.1" />
   </ItemGroup>
 
   <ItemGroup>
diff --git a/Engine/Source/Programs/Horde/Plugins/Build/HordeServer.Build.Tests/Agents/Pools/PoolsControllerTest.cs b/Engine/Source/Programs/Horde/Plugins/Build/HordeServer.Build.Tests/Agents/Pools/PoolsControllerTest.cs
--- a/Engine/Source/Programs/Horde/Plugins/Build/HordeServer.Build.Tests/Agents/Pools/PoolsControllerTest.cs
+++ b/Engine/Source/Programs/Horde/Plugins/Build/HordeServer.Build.Tests/Agents/Pools/PoolsControllerTest.cs
@@ -4,10 +4,7 @@ using EpicGames.Horde.Agents.Pools;
 using HordeServer.Agents.Fleet;
 using HordeServer.Agents.Pools;
 using Microsoft.AspNetCore.Mvc;
-using Microsoft.Extensions.DependencyInjection;
-using Microsoft.Extensions.Options;
 using MongoDB.Bson;
-using MongoDB.Bson.Serialization.Attributes;
 using MongoDB.Driver;
 
 namespace HordeServer.Tests.Agents.Pools
diff --git a/Engine/Source/Programs/Horde/Plugins/Build/HordeServer.Build.Tests/Agents/Relay/AgentRelayClientTest.cs b/Engine/Source/Programs/Horde/Plugins/Build/HordeServer.Build.Tests/Agents/Relay/AgentRelayClientTest.cs
--- a/Engine/Source/Programs/Horde/Plugins/Build/HordeServer.Build.Tests/Agents/Relay/AgentRelayClientTest.cs
+++ b/Engine/Source/Programs/Horde/Plugins/Build/HordeServer.Build.Tests/Agents/Relay/AgentRelayClientTest.cs
@@ -61,7 +61,7 @@ public class AgentRelayClientTest
 	public async Task GetPortMappingsLongPoll_TimesOut_Async()
 	{
 		using CancellationTokenSource cts = new(3000);
-		await Assert.ThrowsExceptionAsync<OperationCanceledException>(() => _relayClient.GetPortMappingsLongPollAsync(cts.Token));
+		await Assert.ThrowsExactlyAsync<OperationCanceledException>(() => _relayClient.GetPortMappingsLongPollAsync(cts.Token));
 	}
 
 	[TestMethod]
diff --git a/Engine/Source/Programs/Horde/Plugins/Build/HordeServer.Build.Tests/Agents/Relay/AgentRelayServiceTests.cs b/Engine/Source/Programs/Horde/Plugins/Build/HordeServer.Build.Tests/Agents/Relay/AgentRelayServiceTests.cs
--- a/Engine/Source/Programs/Horde/Plugins/Build/HordeServer.Build.Tests/Agents/Relay/AgentRelayServiceTests.cs
+++ b/Engine/Source/Programs/Horde/Plugins/Build/HordeServer.Build.Tests/Agents/Relay/AgentRelayServiceTests.cs
@@ -269,7 +269,7 @@ public class AgentRelayServiceTests : BuildTestSetup
 		_service.SetMinMaxPorts(1000, 1003);
 		await AddPortMappingAsync("cluster1", _pm1.LeaseId, null, _pm1.AgentIp, _pm1.Ports);
 		await AddPortMappingAsync("cluster1", _pm3.LeaseId, null, _pm3.AgentIp, _pm3.Ports);
-		await Assert.ThrowsExceptionAsync<ArgumentException>(() => AddPortMappingAsync("cluster1", _pm2.LeaseId, null, _pm2.AgentIp, _pm2.Ports));
+		await Assert.ThrowsExactlyAsync<ArgumentException>(() => AddPortMappingAsync("cluster1", _pm2.LeaseId, null, _pm2.AgentIp, _pm2.Ports));
 	}
 
 	[TestMethod]
@@ -278,7 +278,7 @@ public class AgentRelayServiceTests : BuildTestSetup
 		_service.SetMinMaxPorts(1000, 1003);
 		await AddPortMappingAsync("cluster1", _pm1.LeaseId, null, _pm1.AgentIp, _pm1.Ports);
 		await AddPortMappingAsync("cluster1", _pm3.LeaseId, null, _pm3.AgentIp, _pm3.Ports);
-		await Assert.ThrowsExceptionAsync<ArgumentException>(() => AddPortMappingAsync("cluster1", _pm2.LeaseId, null, _pm2.AgentIp, _pm2.Ports));
+		await Assert.ThrowsExactlyAsync<ArgumentException>(() => AddPortMappingAsync("cluster1", _pm2.LeaseId, null, _pm2.AgentIp, _pm2.Ports));
 		Assert.IsTrue(await _service.RemovePortMappingAsync(_cluster1, LeaseId.Parse(_pm3.LeaseId)));
 		await AddPortMappingAsync("cluster1", _pm2.LeaseId, null, _pm2.AgentIp, _pm2.Ports);
 	}
@@ -337,9 +337,9 @@ public class AgentRelayServiceTests : BuildTestSetup
 		AssertPorts(Expected(2, 1, 4, 7, 6, 3, 9, 0, 8, 5), Actual(10));
 		AssertPorts(Expected(4, 7, 6, 3), Actual(4, [2, 1, 9]));
 		AssertPorts(Expected(4, 7, 6, 3, 0, 8, 5), Actual(7, [2, 1, 9]));
-		Assert.ThrowsException<ArgumentException>(() => Actual(8, [2, 1, 9]));
+		Assert.ThrowsExactly<ArgumentException>(() => Actual(8, [2, 1, 9]));
 		
-		Exception e1 = Assert.ThrowsException<ArgumentException>(() => Actual(11));
+		Exception e1 = Assert.ThrowsExactly<ArgumentException>(() => Actual(11));
 		Assert.IsTrue(e1.Message.Contains("exceeds", StringComparison.InvariantCulture));
 	}
 	
diff --git a/Engine/Source/Programs/Horde/Plugins/Build/HordeServer.Build.Tests/Compute/ComputeServiceTest.cs b/Engine/Source/Programs/Horde/Plugins/Build/HordeServer.Build.Tests/Compute/ComputeServiceTest.cs
--- a/Engine/Source/Programs/Horde/Plugins/Build/HordeServer.Build.Tests/Compute/ComputeServiceTest.cs
+++ b/Engine/Source/Programs/Horde/Plugins/Build/HordeServer.Build.Tests/Compute/ComputeServiceTest.cs
@@ -247,8 +247,8 @@ namespace HordeServer.Tests.Compute
 
 			IAgent agent = await CreateAgentAsync(new PoolId("foo"), properties: ["ComputeIp=11.0.0.1", "ComputePort=5000"]);
 
-			await Assert.ThrowsExceptionAsync<NoComputeResourcesException>(() => ComputeService.TryAllocateResourceAsync(CreateParams("11.0.0.1"), CancellationToken.None));
-			await Assert.ThrowsExceptionAsync<NoComputeResourcesException>(() => ComputeService.TryAllocateResourceAsync(CreateParams("11.0.0.1"), CancellationToken.None));
+			await Assert.ThrowsExactlyAsync<NoComputeResourcesException>(() => ComputeService.TryAllocateResourceAsync(CreateParams("11.0.0.1"), CancellationToken.None));
+			await Assert.ThrowsExactlyAsync<NoComputeResourcesException>(() => ComputeService.TryAllocateResourceAsync(CreateParams("11.0.0.1"), CancellationToken.None));
 
 			ComputeResource resource2 = await ComputeService.TryAllocateResourceAsync(CreateParams("11.0.0.2"), CancellationToken.None);
 			Assert.AreEqual(agent.Id, resource2!.AgentId);
@@ -271,7 +271,7 @@ namespace HordeServer.Tests.Compute
 			ComputeResource resource2 = await ComputeService.TryAllocateResourceAsync(CreateParams("bar"), CancellationToken.None);
 			Assert.AreEqual(agent2.Id, resource2!.AgentId);
 
-			await Assert.ThrowsExceptionAsync<NoComputeResourcesException>(() => ComputeService.TryAllocateResourceAsync(CreateParams("does-not-exist"), CancellationToken.None));
+			await Assert.ThrowsExactlyAsync<NoComputeResourcesException>(() => ComputeService.TryAllocateResourceAsync(CreateParams("does-not-exist"), CancellationToken.None));
 		}
 		
 		[TestMethod]
@@ -297,10 +297,10 @@ namespace HordeServer.Tests.Compute
 			
 			await CreateAgentAsync(new PoolId("foo"), properties: ["ComputeIp=11.0.0.1", "ComputePort=5000"]);
 			await ComputeService.TryAllocateResourceAsync(CreateParams("foo"), CancellationToken.None);
-			NoComputeResourcesException ex1 = await Assert.ThrowsExceptionAsync<NoComputeResourcesException>(() => ComputeService.TryAllocateResourceAsync(CreateParams("foo"), CancellationToken.None));
+			NoComputeResourcesException ex1 = await Assert.ThrowsExactlyAsync<NoComputeResourcesException>(() => ComputeService.TryAllocateResourceAsync(CreateParams("foo"), CancellationToken.None));
 			Assert.AreEqual(NoComputeResourcesException.AllResourcesInUse, ex1.Message);
 			
-			NoComputeResourcesException ex2 = await Assert.ThrowsExceptionAsync<NoComputeResourcesException>(() => ComputeService.TryAllocateResourceAsync(CreateParams("bar"), CancellationToken.None));
+			NoComputeResourcesException ex2 = await Assert.ThrowsExactlyAsync<NoComputeResourcesException>(() => ComputeService.TryAllocateResourceAsync(CreateParams("bar"), CancellationToken.None));
 			Assert.AreEqual(NoComputeResourcesException.NoMatchingResources, ex2.Message);
 		}
 
@@ -329,12 +329,12 @@ namespace HordeServer.Tests.Compute
 				Assert.AreEqual(agent1.Id, (await ComputeService.TryAllocateResourceAsync(arpFoo, CancellationToken.None))!.AgentId);
 
 				AllocateResourceParams arpBar = new(cluster1, _user, ComputeProtocol.Latest, new Requirements { Pool = poolBar.ToString() });
-				await Assert.ThrowsExceptionAsync<NoComputeResourcesException>(() => ComputeService.TryAllocateResourceAsync(arpBar, CancellationToken.None));
+				await Assert.ThrowsExactlyAsync<NoComputeResourcesException>(() => ComputeService.TryAllocateResourceAsync(arpBar, CancellationToken.None));
 			}
 
 			{
 				AllocateResourceParams arpFoo = new(cluster2, _user, ComputeProtocol.Latest, new Requirements { Pool = poolFoo.ToString() });
-				await Assert.ThrowsExceptionAsync<NoComputeResourcesException>(() => ComputeService.TryAllocateResourceAsync(arpFoo, CancellationToken.None));
+				await Assert.ThrowsExactlyAsync<NoComputeResourcesException>(() => ComputeService.TryAllocateResourceAsync(arpFoo, CancellationToken.None));
 
 				AllocateResourceParams arpBar = new(cluster2, _user, ComputeProtocol.Latest, new Requirements { Pool = poolBar.ToString() });
 				Assert.AreEqual(agent2.Id, (await ComputeService.TryAllocateResourceAsync(arpBar, CancellationToken.None))!.AgentId);
@@ -369,7 +369,7 @@ namespace HordeServer.Tests.Compute
 
 			Assert.AreEqual(new ClusterId("compute1"), ComputeService.FindBestComputeClusterId(computeConfig, IPAddress.Parse("11.0.0.1")));
 			Assert.AreEqual(new ClusterId("compute2"), ComputeService.FindBestComputeClusterId(computeConfig, IPAddress.Parse("12.0.1.1")));
-			Assert.ThrowsException<ComputeServiceException>(() => ComputeService.FindBestComputeClusterId(computeConfig, IPAddress.Parse("123.123.123.123")));
+			Assert.ThrowsExactly<ComputeServiceException>(() => ComputeService.FindBestComputeClusterId(computeConfig, IPAddress.Parse("123.123.123.123")));
 
 			ComputeConfig globalConfigCatchAll = new()
 			{
@@ -488,11 +488,11 @@ namespace HordeServer.Tests.Compute
 		[TestMethod]
 		public async Task Connection_UbaCache_CacheServerNotDefined_Async()
 		{
-			ComputeServiceException cse1 = await Assert.ThrowsExceptionAsync<ComputeServiceException>(() => AllocateWithAgentAsync(useUbaCache: true, ubaConfig: null));
+			ComputeServiceException cse1 = await Assert.ThrowsExactlyAsync<ComputeServiceException>(() => AllocateWithAgentAsync(useUbaCache: true, ubaConfig: null));
 			Assert.IsTrue(cse1.Message.StartsWith("No UBA cache server defined for compute cluster", StringComparison.Ordinal));
 			
 			UbaComputeClusterConfig uccc = new() { CacheEndpoint = "localhost:7009", CacheHttpEndpoint = null, CacheHttpCrypto = "aaaaaaaaaaaaaaaaaaaaaaaaaaaaaaaa" };
-			ComputeServiceException cse2 = await Assert.ThrowsExceptionAsync<ComputeServiceException>(() => AllocateWithAgentAsync(useUbaCache: true, ubaConfig: uccc));
+			ComputeServiceException cse2 = await Assert.ThrowsExactlyAsync<ComputeServiceException>(() => AllocateWithAgentAsync(useUbaCache: true, ubaConfig: uccc));
 			Assert.IsTrue(cse2.Message.StartsWith("No UBA cache server defined for compute cluster", StringComparison.Ordinal));
 		}
 		
diff --git a/Engine/Source/Programs/Horde/Plugins/Build/HordeServer.Build.Tests/Fleet/AwsRecyclingFleetManagerTest.cs b/Engine/Source/Programs/Horde/Plugins/Build/HordeServer.Build.Tests/Fleet/AwsRecyclingFleetManagerTest.cs
--- a/Engine/Source/Programs/Horde/Plugins/Build/HordeServer.Build.Tests/Fleet/AwsRecyclingFleetManagerTest.cs
+++ b/Engine/Source/Programs/Horde/Plugins/Build/HordeServer.Build.Tests/Fleet/AwsRecyclingFleetManagerTest.cs
@@ -27,7 +27,7 @@ namespace HordeServer.Tests.Fleet
 			ec2.AddInstance(FakeAmazonEc2.StateStopped, InstanceType.M5Large, FakeAmazonEc2.AzUsEast1A, tagPoolId: PoolId);
 			ec2.SetCapacity(FakeAmazonEc2.AzUsEast1A, InstanceType.M5Large, 1);
 			ec2.OnStartInstances = (_, _) => throw new AmazonEC2Exception("Something bad happened");
-			await Assert.ThrowsExceptionAsync<AmazonEC2Exception>(() => ExpandPoolAsync(ec2.Get(), 1, new()));
+			await Assert.ThrowsExactlyAsync<AmazonEC2Exception>(() => ExpandPoolAsync(ec2.Get(), 1, new()));
 		}
 
 		[TestMethod]
diff --git a/Engine/Source/Programs/Horde/Plugins/Build/HordeServer.Build.Tests/HordeServer.Build.Tests.csproj b/Engine/Source/Programs/Horde/Plugins/Build/HordeServer.Build.Tests/HordeServer.Build.Tests.csproj
--- a/Engine/Source/Programs/Horde/Plugins/Build/HordeServer.Build.Tests/HordeServer.Build.Tests.csproj
+++ b/Engine/Source/Programs/Horde/Plugins/Build/HordeServer.Build.Tests/HordeServer.Build.Tests.csproj
@@ -26,8 +26,8 @@
       <IncludeAssets>runtime; build; native; contentfiles; analyzers; buildtransitive</IncludeAssets>
     </PackageReference>
     <PackageReference Include="Microsoft.NET.Test.Sdk" Version="17.14.1" />
-    <PackageReference Include="MSTest.TestAdapter" Version="3.6.4" />
-    <PackageReference Include="MSTest.TestFramework" Version="3.6.4" />
+    <PackageReference Include="MSTest.TestAdapter" Version="3.11.1" />
+    <PackageReference Include="MSTest.TestFramework" Version="3.11.1" />
     <PackageReference Include="Microsoft.AspNetCore.Mvc.Testing" Version="8.0.20" />
     <PackageReference Include="Moq" Version="4.20.72" />
     <PackageReference Include="Microsoft.Extensions.Caching.Memory" Version="8.0.1" />
diff --git a/Engine/Source/Programs/Horde/Plugins/Build/HordeServer.Build.Tests/Jobs/Bisect/BisectTests.cs b/Engine/Source/Programs/Horde/Plugins/Build/HordeServer.Build.Tests/Jobs/Bisect/BisectTests.cs
--- a/Engine/Source/Programs/Horde/Plugins/Build/HordeServer.Build.Tests/Jobs/Bisect/BisectTests.cs
+++ b/Engine/Source/Programs/Horde/Plugins/Build/HordeServer.Build.Tests/Jobs/Bisect/BisectTests.cs
@@ -193,6 +193,7 @@ namespace HordeServer.Tests.Jobs.Bisect
 		BisectTasksController? _bisectTasksController;
 
 		// setup BisectTaskController with a test user
+		[System.Diagnostics.CodeAnalysis.SuppressMessage("Design", "MSTEST0036:Do not use shadowing", Justification = "<Pending>")]
 		private new BisectTasksController BisectTasksController
 		{
 			get
diff --git a/Engine/Source/Programs/Horde/Plugins/Build/HordeServer.Build.Tests/Jobs/JobCollectionTests.cs b/Engine/Source/Programs/Horde/Plugins/Build/HordeServer.Build.Tests/Jobs/JobCollectionTests.cs
--- a/Engine/Source/Programs/Horde/Plugins/Build/HordeServer.Build.Tests/Jobs/JobCollectionTests.cs
+++ b/Engine/Source/Programs/Horde/Plugins/Build/HordeServer.Build.Tests/Jobs/JobCollectionTests.cs
@@ -38,6 +38,8 @@ namespace HordeServer.Tests.Jobs
 
 		// override JobController with valid user, so we can use REST API in tests
 		private JobsController? _jobsController;
+
+		[System.Diagnostics.CodeAnalysis.SuppressMessage("Design", "MSTEST0036:Do not use shadowing", Justification = "<Pending>")]
 		private new JobsController JobsController
 		{
 			get
diff --git a/Engine/Source/Programs/Horde/Plugins/Build/HordeServer.Build.Tests/Jobs/Schedules/SchedulerTests.cs b/Engine/Source/Programs/Horde/Plugins/Build/HordeServer.Build.Tests/Jobs/Schedules/SchedulerTests.cs
--- a/Engine/Source/Programs/Horde/Plugins/Build/HordeServer.Build.Tests/Jobs/Schedules/SchedulerTests.cs
+++ b/Engine/Source/Programs/Horde/Plugins/Build/HordeServer.Build.Tests/Jobs/Schedules/SchedulerTests.cs
@@ -861,14 +861,13 @@ namespace HordeServer.Tests.Jobs.Schedules
 		}
 
 		[TestMethod]
-		[ExpectedException(typeof(ConditionException))]
 		public async Task ScheduleWithMissingTagsTestAsync()
 		{
 			ScheduleConfig schedule = new()
 			{
 				Condition = Condition.Parse("A == true")
 			};
-			await SetScheduleAsync(schedule);
+			await Assert.ThrowsExactlyAsync<ConditionException>(async () => await SetScheduleAsync(schedule));
 		}
 
 		[TestMethod]
diff --git a/Engine/Source/Programs/Horde/Plugins/Build/HordeServer.Build.Tests/Jobs/TestData/TestDataTests.cs b/Engine/Source/Programs/Horde/Plugins/Build/HordeServer.Build.Tests/Jobs/TestData/TestDataTests.cs
--- a/Engine/Source/Programs/Horde/Plugins/Build/HordeServer.Build.Tests/Jobs/TestData/TestDataTests.cs
+++ b/Engine/Source/Programs/Horde/Plugins/Build/HordeServer.Build.Tests/Jobs/TestData/TestDataTests.cs
@@ -42,7 +42,8 @@ namespace HordeServer.Tests.Jobs.TestData
 
 		TestDataController? _testDataController;
 
-		// override DeviceController with valid user
+		// override TestDataController with valid user
+		[System.Diagnostics.CodeAnalysis.SuppressMessage("Design", "MSTEST0036:Do not use shadowing", Justification = "<Pending>")]
 		private new TestDataController TestDataController
 		{
 			get
diff --git a/Engine/Source/Programs/Horde/Plugins/Build/HordeServer.Build.Tests/Perforce/PerforceSecretsTests.cs b/Engine/Source/Programs/Horde/Plugins/Build/HordeServer.Build.Tests/Perforce/PerforceSecretsTests.cs
--- a/Engine/Source/Programs/Horde/Plugins/Build/HordeServer.Build.Tests/Perforce/PerforceSecretsTests.cs
+++ b/Engine/Source/Programs/Horde/Plugins/Build/HordeServer.Build.Tests/Perforce/PerforceSecretsTests.cs
@@ -73,14 +73,15 @@ public class PerforceSecretsTests : BuildTestSetup
 	}
 
 	[TestMethod]
-	[ExpectedException(typeof(KeyNotFoundException))]
 	[DataRow("horde:secret:notfound.property", DisplayName = "Resolve Perforce Secret, secret not found")]
 	[DataRow("horde:secret:p4secret.notfound", DisplayName = "Resolve Perforce Secret, property not found")]
 	public async Task ResolvePerforceCredentialsWithSecretNotFoundAsync(string value)
 	{
 		PerforceCredentials credentials = new() { Password = value, Ticket = value };
 		SecretConfig secret = new() { Id = new SecretId("p4secret") };
-		try
+		await Assert.ThrowsExactlyAsync<KeyNotFoundException>(async () =>
+		{
+			try
 		{
 			await AddCredentialsAsync(credentials, secret);
 		}
@@ -88,10 +89,10 @@ public class PerforceSecretsTests : BuildTestSetup
 		{
 			throw ae.InnerExceptions[0];
 		}
+		});
 	}
 
 	[TestMethod]
-	[ExpectedException(typeof(KeyNotFoundException))]
 	public async Task ResolvePerforceCredentialsWithSecretMisconfigurationAsync()
 	{
 		PerforceCredentials credentials = new() { Password = "horde:secret:p4secret.property" };
@@ -100,7 +101,9 @@ public class PerforceSecretsTests : BuildTestSetup
 			Id = new SecretId("p4secret"),
 			Sources = [ new ExternalSecretConfig() { Provider = "SecretProviderKeyNotFound" } ]
 		};
-		try
+		await Assert.ThrowsExactlyAsync<KeyNotFoundException>(async () =>
+		{
+			try
 		{
 			await AddCredentialsAsync(credentials, secret);
 		}
@@ -108,5 +111,6 @@ public class PerforceSecretsTests : BuildTestSetup
 		{
 			throw ae.InnerExceptions[0];
 		}
+		});
 	}
 }
diff --git a/Engine/Source/Programs/Horde/Plugins/Build/HordeServer.Build.Tests/Server/RpcServiceTest.cs b/Engine/Source/Programs/Horde/Plugins/Build/HordeServer.Build.Tests/Server/RpcServiceTest.cs
--- a/Engine/Source/Programs/Horde/Plugins/Build/HordeServer.Build.Tests/Server/RpcServiceTest.cs
+++ b/Engine/Source/Programs/Horde/Plugins/Build/HordeServer.Build.Tests/Server/RpcServiceTest.cs
@@ -134,10 +134,10 @@ namespace HordeServer.Tests.Server
 		public async Task CreateSessionTestAsync()
 		{
 			RpcCreateSessionRequest req = new RpcCreateSessionRequest();
-			await Assert.ThrowsExceptionAsync<StructuredRpcException>(() => RpcService.CreateSession(req, _adminContext));
+			await Assert.ThrowsExactlyAsync<StructuredRpcException>(() => RpcService.CreateSession(req, _adminContext));
 
 			req.Id = new AgentId("MyName").ToString();
-			await Assert.ThrowsExceptionAsync<StructuredRpcException>(() => RpcService.CreateSession(req, _adminContext));
+			await Assert.ThrowsExactlyAsync<StructuredRpcException>(() => RpcService.CreateSession(req, _adminContext));
 
 			req.Capabilities = new RpcAgentCapabilities();
 			RpcCreateSessionResponse res = await RpcService.CreateSession(req, _adminContext);
@@ -180,7 +180,7 @@ namespace HordeServer.Tests.Server
 			}
 			else
 			{
-				StructuredRpcException ex = await Assert.ThrowsExceptionAsync<StructuredRpcException>(() => RpcService.CreateAgent(req, context));
+				StructuredRpcException ex = await Assert.ThrowsExactlyAsync<StructuredRpcException>(() => RpcService.CreateAgent(req, context));
 				Assert.AreEqual(StatusCode.PermissionDenied, ex.StatusCode);
 			}
 		}
@@ -204,7 +204,7 @@ namespace HordeServer.Tests.Server
 			}
 			else
 			{
-				StructuredRpcException ex = await Assert.ThrowsExceptionAsync<StructuredRpcException>(() => RpcService.DownloadSoftware(req, responseStream, context));
+				StructuredRpcException ex = await Assert.ThrowsExactlyAsync<StructuredRpcException>(() => RpcService.DownloadSoftware(req, responseStream, context));
 				Assert.AreEqual(expectedStatusCode, ex.StatusCode);
 			}
 		}
@@ -279,7 +279,7 @@ namespace HordeServer.Tests.Server
 			Task call = RpcService.UpdateSession(requestStream, responseStream, _adminContext);
 
 			requestStream.AddMessage(new RpcUpdateSessionRequest { AgentId = "does-not-exist", SessionId = sessionId });
-			StructuredRpcException re = await Assert.ThrowsExceptionAsync<StructuredRpcException>(() => call);
+			StructuredRpcException re = await Assert.ThrowsExactlyAsync<StructuredRpcException>(() => call);
 			Assert.AreEqual(StatusCode.NotFound, re.StatusCode);
 			Assert.IsTrue(re.Message.Contains("Invalid agent name", StringComparison.OrdinalIgnoreCase));
 		}
@@ -301,7 +301,7 @@ namespace HordeServer.Tests.Server
 			Task call = RpcService.UpdateSession(requestStream, responseStream, _adminContext);
 
 			requestStream.AddMessage(new RpcUpdateSessionRequest { AgentId = "does-not-exist", SessionId = sessionId });
-			StructuredRpcException re = await Assert.ThrowsExceptionAsync<StructuredRpcException>(() => call);
+			StructuredRpcException re = await Assert.ThrowsExactlyAsync<StructuredRpcException>(() => call);
 			Assert.AreEqual(StatusCode.NotFound, re.StatusCode);
 			Assert.IsTrue(re.Message.Contains("Invalid agent name", StringComparison.OrdinalIgnoreCase));
 		}
diff --git a/Engine/Source/Programs/Horde/Plugins/Build/HordeServer.Build.Tests/Streams/StreamCollectionTests.cs b/Engine/Source/Programs/Horde/Plugins/Build/HordeServer.Build.Tests/Streams/StreamCollectionTests.cs
--- a/Engine/Source/Programs/Horde/Plugins/Build/HordeServer.Build.Tests/Streams/StreamCollectionTests.cs
+++ b/Engine/Source/Programs/Horde/Plugins/Build/HordeServer.Build.Tests/Streams/StreamCollectionTests.cs
@@ -28,7 +28,7 @@ namespace HordeServer.Tests.Streams
 				Templates = new() { new TemplateRefConfig { Id = new TemplateId("foo") } }
 			};
 
-			Assert.ThrowsException<InvalidStreamException>(() => HordeServer.Streams.StreamCollection.Validate(_streamId, config));
+			Assert.ThrowsExactly<InvalidStreamException>(() => HordeServer.Streams.StreamCollection.Validate(_streamId, config));
 
 			config.Templates.Add(new TemplateRefConfig { Id = new TemplateId("bar") });
 			HordeServer.Streams.StreamCollection.Validate(_streamId, config);
diff --git a/Engine/Source/Programs/Horde/Plugins/Compute/HordeServer.Compute.Tests/Agents/AgentCollectionTests.cs b/Engine/Source/Programs/Horde/Plugins/Compute/HordeServer.Compute.Tests/Agents/AgentCollectionTests.cs
--- a/Engine/Source/Programs/Horde/Plugins/Compute/HordeServer.Compute.Tests/Agents/AgentCollectionTests.cs
+++ b/Engine/Source/Programs/Horde/Plugins/Compute/HordeServer.Compute.Tests/Agents/AgentCollectionTests.cs
@@ -413,7 +413,7 @@ public class AgentCollectionTests : ComputeTestSetup
 	}
 	
 	[TestMethod]
-	[DynamicData(nameof(GetSoftwareToolIdTestData), DynamicDataSourceType.Method)]
+	[DynamicData(nameof(GetSoftwareToolIdTestData))]
 	public async Task GetSoftwareToolIdAsync(string? reportedToolId, AgentSoftwareConfig? softwareConfig, string? expectedToolId)
 	{
 		List<string> props = reportedToolId != null ? [$"{KnownPropertyNames.ToolId}={reportedToolId}"] : [];
diff --git a/Engine/Source/Programs/Horde/Plugins/Compute/HordeServer.Compute.Tests/HordeServer.Compute.Tests.csproj b/Engine/Source/Programs/Horde/Plugins/Compute/HordeServer.Compute.Tests/HordeServer.Compute.Tests.csproj
--- a/Engine/Source/Programs/Horde/Plugins/Compute/HordeServer.Compute.Tests/HordeServer.Compute.Tests.csproj
+++ b/Engine/Source/Programs/Horde/Plugins/Compute/HordeServer.Compute.Tests/HordeServer.Compute.Tests.csproj
@@ -26,8 +26,8 @@
       <IncludeAssets>runtime; build; native; contentfiles; analyzers; buildtransitive</IncludeAssets>
     </PackageReference>
     <PackageReference Include="Microsoft.NET.Test.Sdk" Version="17.14.1" />
-    <PackageReference Include="MSTest.TestAdapter" Version="3.6.4" />
-    <PackageReference Include="MSTest.TestFramework" Version="3.6.4" />
+    <PackageReference Include="MSTest.TestAdapter" Version="3.11.1" />
+    <PackageReference Include="MSTest.TestFramework" Version="3.11.1" />
   </ItemGroup>
 
   <ItemGroup>
diff --git a/Engine/Source/Programs/Horde/Plugins/ExamplePlugin/HordeServer.ExamplePlugin.Tests/HordeServer.ExamplePlugin.Tests.csproj b/Engine/Source/Programs/Horde/Plugins/ExamplePlugin/HordeServer.ExamplePlugin.Tests/HordeServer.ExamplePlugin.Tests.csproj
--- a/Engine/Source/Programs/Horde/Plugins/ExamplePlugin/HordeServer.ExamplePlugin.Tests/HordeServer.ExamplePlugin.Tests.csproj
+++ b/Engine/Source/Programs/Horde/Plugins/ExamplePlugin/HordeServer.ExamplePlugin.Tests/HordeServer.ExamplePlugin.Tests.csproj
@@ -26,8 +26,8 @@
       <IncludeAssets>runtime; build; native; contentfiles; analyzers; buildtransitive</IncludeAssets>
     </PackageReference>
     <PackageReference Include="Microsoft.NET.Test.Sdk" Version="17.14.1" />
-    <PackageReference Include="MSTest.TestAdapter" Version="3.6.4" />
-    <PackageReference Include="MSTest.TestFramework" Version="3.6.4" />
+    <PackageReference Include="MSTest.TestAdapter" Version="3.11.1" />
+    <PackageReference Include="MSTest.TestFramework" Version="3.11.1" />
     <PackageReference Include="Microsoft.AspNetCore.Mvc.Testing" Version="8.0.20" />
     <PackageReference Include="Moq" Version="4.20.72" />
     <PackageReference Include="Microsoft.Extensions.Caching.Memory" Version="8.0.1" />
diff --git a/Engine/Source/Programs/Horde/Plugins/Experimental/HordeServer.Experimental.Tests/HordeServer.Experimental.Tests.csproj b/Engine/Source/Programs/Horde/Plugins/Experimental/HordeServer.Experimental.Tests/HordeServer.Experimental.Tests.csproj
--- a/Engine/Source/Programs/Horde/Plugins/Experimental/HordeServer.Experimental.Tests/HordeServer.Experimental.Tests.csproj
+++ b/Engine/Source/Programs/Horde/Plugins/Experimental/HordeServer.Experimental.Tests/HordeServer.Experimental.Tests.csproj
@@ -26,8 +26,8 @@
       <IncludeAssets>runtime; build; native; contentfiles; analyzers; buildtransitive</IncludeAssets>
     </PackageReference>
     <PackageReference Include="Microsoft.NET.Test.Sdk" Version="17.14.1" />
-    <PackageReference Include="MSTest.TestAdapter" Version="3.6.4" />
-    <PackageReference Include="MSTest.TestFramework" Version="3.6.4" />
+    <PackageReference Include="MSTest.TestAdapter" Version="3.11.1" />
+    <PackageReference Include="MSTest.TestFramework" Version="3.11.1" />
     <PackageReference Include="Microsoft.AspNetCore.Mvc.Testing" Version="8.0.20" />
     <PackageReference Include="Moq" Version="4.20.72" />
     <PackageReference Include="Microsoft.Extensions.Caching.Memory" Version="8.0.1" />
diff --git a/Engine/Source/Programs/Horde/Plugins/Secrets/HordeServer.Secrets.Tests/HordeServer.Secrets.Tests.csproj b/Engine/Source/Programs/Horde/Plugins/Secrets/HordeServer.Secrets.Tests/HordeServer.Secrets.Tests.csproj
--- a/Engine/Source/Programs/Horde/Plugins/Secrets/HordeServer.Secrets.Tests/HordeServer.Secrets.Tests.csproj
+++ b/Engine/Source/Programs/Horde/Plugins/Secrets/HordeServer.Secrets.Tests/HordeServer.Secrets.Tests.csproj
@@ -28,8 +28,8 @@
       <IncludeAssets>runtime; build; native; contentfiles; analyzers; buildtransitive</IncludeAssets>
     </PackageReference>
     <PackageReference Include="Microsoft.NET.Test.Sdk" Version="17.14.1" />
-    <PackageReference Include="MSTest.TestAdapter" Version="3.6.4" />
-    <PackageReference Include="MSTest.TestFramework" Version="3.6.4" />
+    <PackageReference Include="MSTest.TestAdapter" Version="3.11.1" />
+    <PackageReference Include="MSTest.TestFramework" Version="3.11.1" />
   </ItemGroup>
 
   <ItemGroup>
diff --git a/Engine/Source/Programs/Horde/Plugins/Secrets/HordeServer.Secrets.Tests/Secrets/SecretProviderTests.cs b/Engine/Source/Programs/Horde/Plugins/Secrets/HordeServer.Secrets.Tests/Secrets/SecretProviderTests.cs
--- a/Engine/Source/Programs/Horde/Plugins/Secrets/HordeServer.Secrets.Tests/Secrets/SecretProviderTests.cs
+++ b/Engine/Source/Programs/Horde/Plugins/Secrets/HordeServer.Secrets.Tests/Secrets/SecretProviderTests.cs
@@ -260,7 +260,7 @@ namespace HordeServer.Tests.Secrets
 		public async Task SecretNotFoundFromVaultAsync()
 		{
 			SecretCollection secretCollection = CreateSecretCollection(); 
-			InvalidOperationException ex = await Assert.ThrowsExceptionAsync<InvalidOperationException>(() =>
+			InvalidOperationException ex = await Assert.ThrowsExactlyAsync<InvalidOperationException>(() =>
 				secretCollection.GetAsync(new SecretId("MissingSecretUsingHcpVault"), CancellationToken.None)
 			);
 			Assert.AreEqual(ex.Message, "Unable to fetch secret '/v1/secret/data/notfound' from HCP Vault Status = NotFound Body={\"errors\":[]}\n");
diff --git a/Engine/Source/Programs/Horde/Plugins/Secrets/HordeServer.Secrets.Tests/Secrets/SecretResolveTests.cs b/Engine/Source/Programs/Horde/Plugins/Secrets/HordeServer.Secrets.Tests/Secrets/SecretResolveTests.cs
--- a/Engine/Source/Programs/Horde/Plugins/Secrets/HordeServer.Secrets.Tests/Secrets/SecretResolveTests.cs
+++ b/Engine/Source/Programs/Horde/Plugins/Secrets/HordeServer.Secrets.Tests/Secrets/SecretResolveTests.cs
@@ -64,19 +64,17 @@ namespace HordeServer.Tests.Secrets
 		}
 
 		[TestMethod]
-		[ExpectedException(typeof(KeyNotFoundException))]
 		[DataRow("horde:secret:doesNotExist.password", DisplayName = "Secret does not exist")]
 		[DataRow("horde:secret:mySecret.doesNotExist", DisplayName = "Secret property does not exist")]
 		public async Task ResolveSecretNotFoundAsync(string value)
 		{
-			_ = await s_secretCollection.ResolveAsync(value);
+			await Assert.ThrowsExactlyAsync<KeyNotFoundException>(async () => _ = await s_secretCollection.ResolveAsync(value));
 		}
 
 		[TestMethod]
-		[ExpectedException(typeof(KeyNotFoundException))]
 		public async Task ResolveSecretWithMisconfiguredConfigAsync()
 		{
-			_ = await s_secretCollection.ResolveAsync("horde:secret:misconfiguredSecret.property");
+			await Assert.ThrowsExactlyAsync<KeyNotFoundException>(async () => _ = await s_secretCollection.ResolveAsync("horde:secret:misconfiguredSecret.property"));
 		}
 	}
 }
diff --git a/Engine/Source/Programs/Horde/Plugins/Secrets/HordeServer.Secrets.Tests/Secrets/SecretTests.cs b/Engine/Source/Programs/Horde/Plugins/Secrets/HordeServer.Secrets.Tests/Secrets/SecretTests.cs
--- a/Engine/Source/Programs/Horde/Plugins/Secrets/HordeServer.Secrets.Tests/Secrets/SecretTests.cs
+++ b/Engine/Source/Programs/Horde/Plugins/Secrets/HordeServer.Secrets.Tests/Secrets/SecretTests.cs
@@ -57,10 +57,9 @@ namespace HordeServer.Tests.Secrets
 		}
 
 		[TestMethod]
-		[ExpectedException(typeof(InvalidCastException))]
 		public async Task GetSecretWithInvalidConfigAsync()
 		{
-			_ = await s_secretCollection.GetAsync(new SecretId(), new BuildConfig());
+			await Assert.ThrowsExactlyAsync<InvalidCastException>(async () => _ = await s_secretCollection.GetAsync(new SecretId(), new BuildConfig()));
 		}
 	}
 }
diff --git a/Engine/Source/Programs/Horde/Plugins/Storage/HordeServer.Storage.Tests/HordeServer.Storage.Tests.csproj b/Engine/Source/Programs/Horde/Plugins/Storage/HordeServer.Storage.Tests/HordeServer.Storage.Tests.csproj
--- a/Engine/Source/Programs/Horde/Plugins/Storage/HordeServer.Storage.Tests/HordeServer.Storage.Tests.csproj
+++ b/Engine/Source/Programs/Horde/Plugins/Storage/HordeServer.Storage.Tests/HordeServer.Storage.Tests.csproj
@@ -26,8 +26,8 @@
       <IncludeAssets>runtime; build; native; contentfiles; analyzers; buildtransitive</IncludeAssets>
     </PackageReference>
     <PackageReference Include="Microsoft.NET.Test.Sdk" Version="17.14.1" />
-    <PackageReference Include="MSTest.TestAdapter" Version="3.6.4" />
-    <PackageReference Include="MSTest.TestFramework" Version="3.6.4" />
+    <PackageReference Include="MSTest.TestAdapter" Version="3.11.1" />
+    <PackageReference Include="MSTest.TestFramework" Version="3.11.1" />
   </ItemGroup>
 
   <ItemGroup>
diff --git a/Engine/Source/Programs/Horde/Plugins/Tools/HordeServer.Tools.Tests/HordeServer.Tools.Tests.csproj b/Engine/Source/Programs/Horde/Plugins/Tools/HordeServer.Tools.Tests/HordeServer.Tools.Tests.csproj
--- a/Engine/Source/Programs/Horde/Plugins/Tools/HordeServer.Tools.Tests/HordeServer.Tools.Tests.csproj
+++ b/Engine/Source/Programs/Horde/Plugins/Tools/HordeServer.Tools.Tests/HordeServer.Tools.Tests.csproj
@@ -26,8 +26,8 @@
       <IncludeAssets>runtime; build; native; contentfiles; analyzers; buildtransitive</IncludeAssets>
     </PackageReference>
     <PackageReference Include="Microsoft.NET.Test.Sdk" Version="17.14.1" />
-    <PackageReference Include="MSTest.TestAdapter" Version="3.6.4" />
-    <PackageReference Include="MSTest.TestFramework" Version="3.6.4" />
+    <PackageReference Include="MSTest.TestAdapter" Version="3.11.1" />
+    <PackageReference Include="MSTest.TestFramework" Version="3.11.1" />
   </ItemGroup>
 
   <ItemGroup>
diff --git a/Engine/Source/Programs/Shared/EpicGames.Perforce.Fixture/EpicGames.Perforce.Fixture.csproj b/Engine/Source/Programs/Shared/EpicGames.Perforce.Fixture/EpicGames.Perforce.Fixture.csproj
--- a/Engine/Source/Programs/Shared/EpicGames.Perforce.Fixture/EpicGames.Perforce.Fixture.csproj
+++ b/Engine/Source/Programs/Shared/EpicGames.Perforce.Fixture/EpicGames.Perforce.Fixture.csproj
@@ -12,10 +12,10 @@
   <ItemGroup>
     <PackageReference Include="Microsoft.Extensions.Logging.Console" Version="8.0.1" />
     <PackageReference Include="Microsoft.NET.Test.Sdk" Version="17.14.1" />
-    <PackageReference Include="MSTest.TestAdapter" Version="3.6.4">
+    <PackageReference Include="MSTest.TestAdapter" Version="3.11.1">
       <TreatAsUsed>true</TreatAsUsed>
     </PackageReference>
-    <PackageReference Include="MSTest.TestFramework" Version="3.6.4" />
+    <PackageReference Include="MSTest.TestFramework" Version="3.11.1" />
     <PackageReference Include="coverlet.collector" Version="6.0.4">
       <PrivateAssets>all</PrivateAssets>
       <IncludeAssets>runtime; build; native; contentfiles; analyzers; buildtransitive</IncludeAssets>
