Horde: Update Microsoft.AspNetCore.Mvc.Testing

#rnx
#Synced-CL 50981627

diff --git a/Engine/Source/Programs/Horde/HordeServer.Tests.Shared/HordeServer.Tests.Shared.csproj b/Engine/Source/Programs/Horde/HordeServer.Tests.Shared/HordeServer.Tests.Shared.csproj
--- a/Engine/Source/Programs/Horde/HordeServer.Tests.Shared/HordeServer.Tests.Shared.csproj
+++ b/Engine/Source/Programs/Horde/HordeServer.Tests.Shared/HordeServer.Tests.Shared.csproj
@@ -23,7 +23,7 @@
   </ItemGroup>
 
   <ItemGroup>
-    <PackageReference Include="Microsoft.AspNetCore.Mvc.Testing" Version="8.0.20" />
+    <PackageReference Include="Microsoft.AspNetCore.Mvc.Testing" Version="8.0.24" />
     <PackageReference Include="Microsoft.NET.Test.Sdk" Version="17.14.1" />
     <PackageReference Include="Moq" Version="4.20.72" />
     <PackageReference Include="MSTest.TestAdapter" Version="3.11.1" />
diff --git a/Engine/Source/Programs/Horde/HordeServer.Tests/HordeServer.Tests.csproj b/Engine/Source/Programs/Horde/HordeServer.Tests/HordeServer.Tests.csproj
--- a/Engine/Source/Programs/Horde/HordeServer.Tests/HordeServer.Tests.csproj
+++ b/Engine/Source/Programs/Horde/HordeServer.Tests/HordeServer.Tests.csproj
@@ -20,7 +20,7 @@
   </PropertyGroup>
 
   <ItemGroup>
-    <PackageReference Include="Microsoft.AspNetCore.Mvc.Testing" Version="8.0.20" />
+    <PackageReference Include="Microsoft.AspNetCore.Mvc.Testing" Version="8.0.24" />
     <PackageReference Include="Microsoft.NET.Test.Sdk" Version="17.14.1" />
     <PackageReference Include="Moq" Version="4.20.72" />
     <PackageReference Include="MSTest.TestAdapter" Version="3.11.1" />
diff --git a/Engine/Source/Programs/Horde/Plugins/Build/HordeServer.Build.Tests/HordeServer.Build.Tests.csproj b/Engine/Source/Programs/Horde/Plugins/Build/HordeServer.Build.Tests/HordeServer.Build.Tests.csproj
--- a/Engine/Source/Programs/Horde/Plugins/Build/HordeServer.Build.Tests/HordeServer.Build.Tests.csproj
+++ b/Engine/Source/Programs/Horde/Plugins/Build/HordeServer.Build.Tests/HordeServer.Build.Tests.csproj
@@ -28,7 +28,7 @@
     <PackageReference Include="Microsoft.NET.Test.Sdk" Version="17.14.1" />
     <PackageReference Include="MSTest.TestAdapter" Version="3.11.1" />
     <PackageReference Include="MSTest.TestFramework" Version="3.11.1" />
-    <PackageReference Include="Microsoft.AspNetCore.Mvc.Testing" Version="8.0.20" />
+    <PackageReference Include="Microsoft.AspNetCore.Mvc.Testing" Version="8.0.24" />
     <PackageReference Include="Moq" Version="4.20.72" />
     <PackageReference Include="Microsoft.Extensions.Caching.Memory" Version="8.0.1" />
   </ItemGroup>
diff --git a/Engine/Source/Programs/Horde/Plugins/ExamplePlugin/HordeServer.ExamplePlugin.Tests/HordeServer.ExamplePlugin.Tests.csproj b/Engine/Source/Programs/Horde/Plugins/ExamplePlugin/HordeServer.ExamplePlugin.Tests/HordeServer.ExamplePlugin.Tests.csproj
--- a/Engine/Source/Programs/Horde/Plugins/ExamplePlugin/HordeServer.ExamplePlugin.Tests/HordeServer.ExamplePlugin.Tests.csproj
+++ b/Engine/Source/Programs/Horde/Plugins/ExamplePlugin/HordeServer.ExamplePlugin.Tests/HordeServer.ExamplePlugin.Tests.csproj
@@ -28,7 +28,7 @@
     <PackageReference Include="Microsoft.NET.Test.Sdk" Version="17.14.1" />
     <PackageReference Include="MSTest.TestAdapter" Version="3.11.1" />
     <PackageReference Include="MSTest.TestFramework" Version="3.11.1" />
-    <PackageReference Include="Microsoft.AspNetCore.Mvc.Testing" Version="8.0.20" />
+    <PackageReference Include="Microsoft.AspNetCore.Mvc.Testing" Version="8.0.24" />
     <PackageReference Include="Moq" Version="4.20.72" />
     <PackageReference Include="Microsoft.Extensions.Caching.Memory" Version="8.0.1" />
   </ItemGroup>
diff --git a/Engine/Source/Programs/Horde/Plugins/Experimental/HordeServer.Experimental.Tests/HordeServer.Experimental.Tests.csproj b/Engine/Source/Programs/Horde/Plugins/Experimental/HordeServer.Experimental.Tests/HordeServer.Experimental.Tests.csproj
--- a/Engine/Source/Programs/Horde/Plugins/Experimental/HordeServer.Experimental.Tests/HordeServer.Experimental.Tests.csproj
+++ b/Engine/Source/Programs/Horde/Plugins/Experimental/HordeServer.Experimental.Tests/HordeServer.Experimental.Tests.csproj
@@ -28,7 +28,7 @@
     <PackageReference Include="Microsoft.NET.Test.Sdk" Version="17.14.1" />
     <PackageReference Include="MSTest.TestAdapter" Version="3.11.1" />
     <PackageReference Include="MSTest.TestFramework" Version="3.11.1" />
-    <PackageReference Include="Microsoft.AspNetCore.Mvc.Testing" Version="8.0.20" />
+    <PackageReference Include="Microsoft.AspNetCore.Mvc.Testing" Version="8.0.24" />
     <PackageReference Include="Moq" Version="4.20.72" />
     <PackageReference Include="Microsoft.Extensions.Caching.Memory" Version="8.0.1" />
   </ItemGroup>
