Horde: Add redis-backed queue infrastructure for auto-RCA
Add RcaQueueService for persistent RCA processing with retry support, replacing fire-and-forget Task.Run in IssueService.

#Synced-CL 0

diff --git a/Engine/Source/Programs/Horde/HordeServer.Shared/RCA/IRcaEnqueuer.cs b/Engine/Source/Programs/Horde/HordeServer.Shared/RCA/IRcaEnqueuer.cs
new file mode 100644
--- /dev/null
+++ b/Engine/Source/Programs/Horde/HordeServer.Shared/RCA/IRcaEnqueuer.cs
@@ -0,0 +1,23 @@
+// Copyright Epic Games, Inc. All Rights Reserved.
+
+using EpicGames.Horde.Issues;
+using EpicGames.Horde.Streams;
+
+namespace HordeServer.RCA
+{
+	/// <summary>
+	/// Interface for enqueuing RCA processing requests.
+	/// Implementations persist the request and process it asynchronously.
+	/// </summary>
+	public interface IRcaEnqueuer
+	{
+		/// <summary>
+		/// Enqueue an RCA processing request for the given issue.
+		/// </summary>
+		/// <param name="issueId">The issue ID to analyze</param>
+		/// <param name="streamId">The stream the issue belongs to</param>
+		/// <param name="workflowId">The workflow ID (null if not in a workflow)</param>
+		/// <param name="cancellationToken">Cancellation token</param>
+		Task EnqueueAsync(int issueId, StreamId streamId, WorkflowId? workflowId, CancellationToken cancellationToken);
+	}
+}
diff --git a/Engine/Source/Programs/Horde/Plugins/Build/HordeServer.Build/BuildPlugin.cs b/Engine/Source/Programs/Horde/Plugins/Build/HordeServer.Build/BuildPlugin.cs
--- a/Engine/Source/Programs/Horde/Plugins/Build/HordeServer.Build/BuildPlugin.cs
+++ b/Engine/Source/Programs/Horde/Plugins/Build/HordeServer.Build/BuildPlugin.cs
@@ -33,6 +33,7 @@ using HordeServer.Notifications;
 using HordeServer.Notifications.Sinks;
 using HordeServer.VersionControl.Perforce;
 using HordeServer.Plugins;
+using HordeServer.RCA;
 using HordeServer.Replicators;
 using HordeServer.Streams;
 using HordeServer.Tasks;
@@ -120,6 +121,8 @@ namespace HordeServer
 			services.AddSingleton<IAuditLog<TestId>>(sp => sp.GetRequiredService<IAuditLogFactory<TestId>>().Create("TestHealthAudits.Log", "TestId"));
 
 			services.AddSingleton<IssueService>();
+			services.AddSingleton<RcaQueueService>();
+			services.AddSingleton<IRcaEnqueuer>(sp => sp.GetRequiredService<RcaQueueService>());
 			services.AddSingleton<JobService>();
 			services.AddSingleton<IAuditLog<JobId>>(sp => sp.GetRequiredService<IAuditLogFactory<JobId>>().Create("Jobs.Log", "JobId"));
 			services.AddSingleton<ILogExtAuthProvider>(sp => sp.GetRequiredService<JobService>());
@@ -140,6 +143,9 @@ namespace HordeServer
 			// Issues need to be assigned asynchronously on any pod
 			services.AddHostedService(provider => provider.GetRequiredService<IssueService>());
 
+			// RCA queue uses shared ticker for multi-instance coordination
+			services.AddHostedService(provider => provider.GetRequiredService<RcaQueueService>());
+
 			services.AddHostedService(provider => provider.GetRequiredService<IExternalIssueService>());
 
 			// Task sources. Order of registration is important here; it dictates the priority in which sources are served.
diff --git a/Engine/Source/Programs/Horde/Plugins/Build/HordeServer.Build/Issues/IssueService.cs b/Engine/Source/Programs/Horde/Plugins/Build/HordeServer.Build/Issues/IssueService.cs
--- a/Engine/Source/Programs/Horde/Plugins/Build/HordeServer.Build/Issues/IssueService.cs
+++ b/Engine/Source/Programs/Horde/Plugins/Build/HordeServer.Build/Issues/IssueService.cs
@@ -249,13 +249,12 @@ namespace HordeServer.Issues
 		public IReadOnlyList<IIssueDetails> CachedOpenIssues => _cachedIssues;
 
 		readonly IOptionsMonitor<BuildConfig> _buildConfig;
-		readonly IRcaProcessor? _rcaProcessor;
-		readonly IHostApplicationLifetime _lifetime;
+		readonly IRcaEnqueuer? _rcaEnqueuer;
 
 		/// <summary>
 		/// Constructor
 		/// </summary>
-		public IssueService(JobService jobService, IIssueCollection issueCollection, ICommitService commitService, IJobStepRefCollection jobStepRefs, IUserCollection userCollection, ILogCollection logCollection, IClock clock, IOptionsMonitor<BuildConfig> buildConfig, Tracer tracer, ILogger<IssueService> logger, IHostApplicationLifetime lifetime, IRcaProcessor? rcaProcessor = null)
+		public IssueService(JobService jobService, IIssueCollection issueCollection, ICommitService commitService, IJobStepRefCollection jobStepRefs, IUserCollection userCollection, ILogCollection logCollection, IClock clock, IOptionsMonitor<BuildConfig> buildConfig, Tracer tracer, ILogger<IssueService> logger, IRcaEnqueuer? rcaEnqueuer = null)
 		{
 			Type[] issueTypes = Assembly.GetExecutingAssembly().GetTypes().Where(x => !x.IsAbstract && typeof(IIssue).IsAssignableFrom(x)).ToArray();
 			foreach (Type issueType in issueTypes)
@@ -277,8 +276,7 @@ namespace HordeServer.Issues
 			_buildConfig = buildConfig;
 			_tracer = tracer;
 			_logger = logger;
-			_lifetime = lifetime;
-			_rcaProcessor = rcaProcessor;
+			_rcaEnqueuer = rcaEnqueuer;
 
 			// Find all the issue handler types
 			List<(Type Type, IssueHandlerAttribute Attribute)> handlerTypes = new List<(Type, IssueHandlerAttribute)>();
@@ -1832,28 +1830,10 @@ namespace HordeServer.Issues
 				issue = await _issueCollection.AddIssueAsync(summary, cancellationToken);
 				LoggerUtils.MultiLog([(_logger, LogLevel.Information), (issueLogger, LogLevel.Information)], "Created issue {IssueId}", issue.Id);
 
-				// Fire-and-forget RCA processing for genuinely new issues only
-				// Processor handles workflow eligibility and rate limiting internally
-				// Uses ApplicationStopping token to cancel in-flight requests during shutdown
-				if (_rcaProcessor != null)
+				if (_rcaEnqueuer != null)
 				{
 					WorkflowId? workflowId = span.FirstFailure.Annotations?.WorkflowId;
-					CancellationToken stoppingToken = _lifetime.ApplicationStopping;
-					_ = Task.Run(async () =>
-					{
-						try
-						{
-							await _rcaProcessor.ProcessRcaAsync(issue.Id, span.StreamId, workflowId, stoppingToken);
-						}
-						catch (OperationCanceledException) when (stoppingToken.IsCancellationRequested)
-						{
-							_logger.LogDebug("RCA processing cancelled during shutdown for issue {IssueId}", issue.Id);
-						}
-						catch (Exception ex)
-						{
-							_logger.LogError(ex, "RCA processing failed for issue {IssueId}", issue.Id);
-						}
-					}, CancellationToken.None);
+					await _rcaEnqueuer.EnqueueAsync(issue.Id, span.StreamId, workflowId, cancellationToken);
 				}
 			}
 			else
diff --git a/Engine/Source/Programs/Horde/Plugins/Build/HordeServer.Build/Issues/RcaQueueService.cs b/Engine/Source/Programs/Horde/Plugins/Build/HordeServer.Build/Issues/RcaQueueService.cs
new file mode 100644
--- /dev/null
+++ b/Engine/Source/Programs/Horde/Plugins/Build/HordeServer.Build/Issues/RcaQueueService.cs
@@ -0,0 +1,214 @@
+// Copyright Epic Games, Inc. All Rights Reserved.
+
+using EpicGames.Core;
+using EpicGames.Horde.Issues;
+using EpicGames.Horde.Logs;
+using EpicGames.Horde.Streams;
+using EpicGames.Redis;
+using HordeServer.RCA;
+using HordeServer.Server;
+using HordeServer.Streams;
+using HordeServer.Utilities;
+using Microsoft.Extensions.Hosting;
+using Microsoft.Extensions.Logging;
+using Microsoft.Extensions.Options;
+
+namespace HordeServer.Issues
+{
+	/// <summary>
+	/// Redis-backed queue for RCA processing requests.
+	/// Ensures every eligible issue gets processed with retry support and crash resilience.
+	/// </summary>
+	public sealed class RcaQueueService : IRcaEnqueuer, IHostedService, IAsyncDisposable
+	{
+		/// <summary>
+		/// Maximum number of retry attempts before discarding a queue item
+		/// </summary>
+		const int MaxRetries = 3;
+
+		/// <summary>
+		/// Minimum age before a queue item is processed, to allow audit logs to be emitted after issue creation
+		/// </summary>
+		static readonly TimeSpan s_minimumAge = TimeSpan.FromSeconds(5);
+
+		/// <summary>
+		/// Item stored in the Redis RCA queue
+		/// </summary>
+		[RedisConverter(typeof(RedisJsonConverter<>))]
+		record class RcaQueueItem(int IssueId, string StreamId, string? WorkflowId, int RetryCount, DateTime EnqueuedUtc);
+
+		static readonly RedisListKey<RcaQueueItem> s_rcaQueue = "rca/queue";
+
+		readonly IRedisService _redisService;
+		readonly IRcaProcessor? _rcaProcessor;
+		readonly IIssueCollection _issueCollection;
+		readonly IOptionsMonitor<BuildConfig> _buildConfig;
+		readonly ITicker _ticker;
+		readonly ILogger<RcaQueueService> _logger;
+
+		/// <summary>
+		/// Constructor
+		/// </summary>
+		public RcaQueueService(IRedisService redisService, IIssueCollection issueCollection, IClock clock, IOptionsMonitor<BuildConfig> buildConfig, ILogger<RcaQueueService> logger, IRcaProcessor? rcaProcessor = null)
+		{
+			_redisService = redisService;
+			_rcaProcessor = rcaProcessor;
+			_issueCollection = issueCollection;
+			_buildConfig = buildConfig;
+			_ticker = clock.AddSharedTicker<RcaQueueService>(TimeSpan.FromSeconds(30.0), ProcessQueueAsync, logger);
+			_logger = logger;
+		}
+
+		/// <inheritdoc/>
+		public ValueTask DisposeAsync() => _ticker.DisposeAsync();
+
+		/// <inheritdoc/>
+		public Task StartAsync(CancellationToken cancellationToken) => _ticker.StartAsync();
+
+		/// <inheritdoc/>
+		public Task StopAsync(CancellationToken cancellationToken) => _ticker.StopAsync();
+
+		/// <inheritdoc/>
+		public async Task EnqueueAsync(int issueId, StreamId streamId, WorkflowId? workflowId, CancellationToken cancellationToken)
+		{
+			if (_rcaProcessor == null)
+			{
+				return;
+			}
+
+			int maxConcurrent = _buildConfig.CurrentValue.MaxConcurrentRcaRequests;
+			if (maxConcurrent <= 0)
+			{
+				return;
+			}
+
+			RcaQueueItem item = new RcaQueueItem(issueId, streamId.ToString(), workflowId?.ToString(), 0, DateTime.UtcNow);
+			await _redisService.GetDatabase().ListRightPushAsync(s_rcaQueue, item);
+
+			_logger.LogDebug("Enqueued RCA processing for issue {IssueId} in stream {StreamId}", issueId, streamId);
+		}
+
+		async ValueTask ProcessQueueAsync(CancellationToken cancellationToken)
+		{
+			if (_rcaProcessor == null)
+			{
+				return;
+			}
+
+			int maxConcurrent = _buildConfig.CurrentValue.MaxConcurrentRcaRequests;
+			if (maxConcurrent <= 0)
+			{
+				return;
+			}
+
+			long count = await _redisService.GetDatabase().ListLengthAsync(s_rcaQueue);
+			if (count == 0)
+			{
+				return;
+			}
+
+			// Dequeue up to maxConcurrent items, re-queuing any that are too young
+			int batchSize = (int)Math.Min(count, maxConcurrent);
+			List<RcaQueueItem> eligibleItems = new List<RcaQueueItem>(batchSize);
+			for (int i = 0; i < batchSize; i++)
+			{
+				RcaQueueItem item = await _redisService.GetDatabase().ListLeftPopAsync(s_rcaQueue);
+				if (item == null)
+				{
+					break;
+				}
+
+				// Re-queue items that haven't aged enough for audit logs to be emitted
+				if (DateTime.UtcNow - item.EnqueuedUtc < s_minimumAge)
+				{
+					await _redisService.GetDatabase().ListRightPushAsync(s_rcaQueue, item);
+					continue;
+				}
+
+				if (!IsEligibleForRca(item))
+				{
+					_logger.LogDebug("Discarding ineligible RCA queue item for issue {IssueId}", item.IssueId);
+					continue;
+				}
+
+				eligibleItems.Add(item);
+			}
+
+			if (eligibleItems.Count == 0)
+			{
+				return;
+			}
+
+			_logger.LogInformation("Processing {Count} RCA queue items", eligibleItems.Count);
+
+			// Process batch in parallel
+			Task[] tasks = new Task[eligibleItems.Count];
+			for (int i = 0; i < eligibleItems.Count; i++)
+			{
+				tasks[i] = ProcessItemAsync(eligibleItems[i], cancellationToken);
+			}
+			await Task.WhenAll(tasks);
+		}
+
+		async Task ProcessItemAsync(RcaQueueItem item, CancellationToken cancellationToken)
+		{
+			StreamId streamId = new StreamId(item.StreamId);
+			WorkflowId? workflowId = item.WorkflowId != null ? new WorkflowId(item.WorkflowId) : null;
+
+			try
+			{
+				await _rcaProcessor!.ProcessRcaAsync(item.IssueId, streamId, workflowId, cancellationToken);
+			}
+			catch (OperationCanceledException) when (cancellationToken.IsCancellationRequested)
+			{
+				// Shutdown - re-queue so work survives restart
+				await _redisService.GetDatabase().ListRightPushAsync(s_rcaQueue, item);
+				_logger.LogDebug("Re-queued RCA item for issue {IssueId} due to shutdown", item.IssueId);
+			}
+			catch (Exception ex)
+			{
+				ILogger issueLogger = _issueCollection.GetLogger(item.IssueId);
+
+				if (item.RetryCount + 1 >= MaxRetries)
+				{
+					// Permanent failure - discard
+					LoggerUtils.MultiLog(
+						[(_logger, LogLevel.Error), (issueLogger, LogLevel.Information)],
+						ex, "RCA processing permanently failed for issue {IssueId} after {MaxRetries} attempts", item.IssueId, MaxRetries);
+				}
+				else
+				{
+					// Transient failure - re-queue with incremented retry count
+					RcaQueueItem retryItem = item with { RetryCount = item.RetryCount + 1 };
+					await _redisService.GetDatabase().ListRightPushAsync(s_rcaQueue, retryItem);
+
+					LoggerUtils.MultiLog(
+						[(_logger, LogLevel.Warning), (issueLogger, LogLevel.Information)],
+						ex, "RCA processing failed for issue {IssueId} (attempt {Attempt}/{MaxRetries}), re-queued", item.IssueId, item.RetryCount + 1, MaxRetries);
+				}
+			}
+		}
+
+		bool IsEligibleForRca(RcaQueueItem item)
+		{
+			if (item.WorkflowId == null)
+			{
+				return false;
+			}
+
+			StreamId streamId = new StreamId(item.StreamId);
+			if (!_buildConfig.CurrentValue.TryGetStream(streamId, out StreamConfig? streamConfig))
+			{
+				return false;
+			}
+
+			WorkflowId workflowId = new WorkflowId(item.WorkflowId);
+			if (!streamConfig.TryGetWorkflow(workflowId, out WorkflowConfig? workflow))
+			{
+				return false;
+			}
+
+			return workflow.AutoRcaEnabled;
+		}
+	}
+}
