Horde: Fix IDE0005 warning

#rnx
#Synced-CL 0

diff --git a/Engine/Source/Programs/Horde/Plugins/Build/HordeServer.Build/Issues/RcaQueueService.cs b/Engine/Source/Programs/Horde/Plugins/Build/HordeServer.Build/Issues/RcaQueueService.cs
--- a/Engine/Source/Programs/Horde/Plugins/Build/HordeServer.Build/Issues/RcaQueueService.cs
+++ b/Engine/Source/Programs/Horde/Plugins/Build/HordeServer.Build/Issues/RcaQueueService.cs
@@ -8,7 +8,6 @@ using EpicGames.Redis;
 using HordeServer.RCA;
 using HordeServer.Server;
 using HordeServer.Streams;
-using HordeServer.Utilities;
 using Microsoft.Extensions.Hosting;
 using Microsoft.Extensions.Logging;
 using Microsoft.Extensions.Options;
