Horde: Fix PerforceServiceCache not picking up stream view changes
Avoid view mismatch between cache and commit creation by passing pre-computed files to CreateCommitAsync instead of re-resolving the stream view.
Update existing ClusterTicker stream lists every 30s to reflect stream definition changes.

#Synced-CL 0

diff --git a/Engine/Source/Programs/Horde/Plugins/Build/HordeServer.Build/VersionControl/Perforce/PerforceService.cs b/Engine/Source/Programs/Horde/Plugins/Build/HordeServer.Build/VersionControl/Perforce/PerforceService.cs
--- a/Engine/Source/Programs/Horde/Plugins/Build/HordeServer.Build/VersionControl/Perforce/PerforceService.cs
+++ b/Engine/Source/Programs/Horde/Plugins/Build/HordeServer.Build/VersionControl/Perforce/PerforceService.cs
@@ -727,16 +727,24 @@ namespace HordeServer.VersionControl.Perforce
 			return new Commit(this, streamConfig, number, originalChange, authorUser.Id, ownerUser.Id, description, basePath ?? String.Empty, dateUtc);
 		}
 
-		protected async ValueTask<ICommit> CreateCommitAsync(IPooledPerforceConnection perforce, StreamConfig streamConfig, DescribeRecord describeRecord, bool hasAllFiles, InfoRecord serverInfo, CancellationToken cancellationToken)
+		protected async ValueTask<ICommit> CreateCommitAsync(IPooledPerforceConnection perforce, StreamConfig streamConfig, DescribeRecord describeRecord, bool hasAllFiles, InfoRecord serverInfo, CancellationToken cancellationToken, List<string>? preComputedFiles = null)
 		{
 			DateTime timeUtc = describeRecord.Time.Kind == DateTimeKind.Utc ? describeRecord.Time : new DateTime(describeRecord.Time.Ticks - serverInfo.TimeZoneOffsetSecs * TimeSpan.TicksPerSecond, DateTimeKind.Utc);
 
 			Commit commit = await CreateCommitInternalAsync(streamConfig, describeRecord.Number, describeRecord.User, describeRecord.Description, describeRecord.Path, timeUtc, cancellationToken);
 
-			List<string> files = await perforce.GetStreamFilesAsync(streamConfig, describeRecord, cancellationToken);
-			if (files.Count == 0 && describeRecord.Files.Count > 0)
+			List<string> files;
+			if (preComputedFiles != null)
 			{
-				throw new PerforceException($"Changelist {commit.Number} does not contain any files in {streamConfig.Id}");
+				files = preComputedFiles;
+			}
+			else
+			{
+				files = await perforce.GetStreamFilesAsync(streamConfig, describeRecord, cancellationToken);
+				if (files.Count == 0 && describeRecord.Files.Count > 0)
+				{
+					throw new PerforceException($"Changelist {commit.Number} does not contain any files in {streamConfig.Id}");
+				}
 			}
 			commit.SetFiles(files, hasAllFiles);
 
diff --git a/Engine/Source/Programs/Horde/Plugins/Build/HordeServer.Build/VersionControl/Perforce/PerforceServiceCache.cs b/Engine/Source/Programs/Horde/Plugins/Build/HordeServer.Build/VersionControl/Perforce/PerforceServiceCache.cs
--- a/Engine/Source/Programs/Horde/Plugins/Build/HordeServer.Build/VersionControl/Perforce/PerforceServiceCache.cs
+++ b/Engine/Source/Programs/Horde/Plugins/Build/HordeServer.Build/VersionControl/Perforce/PerforceServiceCache.cs
@@ -108,12 +108,19 @@ namespace HordeServer.VersionControl.Perforce
 		{
 			public string ClusterName { get; }
 			public BackgroundTask? Task { get; set; }
-			public List<StreamInfo> Streams { get; set; }
+
+			// Volatile ensures the reference assignment is immediately visible to reader threads. The list contents are never mutated after assignment.
+			private volatile List<StreamInfo> _streams;
+			public List<StreamInfo> Streams
+			{
+				get => _streams;
+				set => _streams = value;
+			}
 
 			public ClusterTicker(string name, List<StreamInfo> streamInfos)
 			{
 				ClusterName = name;
-				Streams = streamInfos;
+				_streams = streamInfos;
 			}
 
 			public async ValueTask DisposeAsync()
@@ -318,7 +325,7 @@ namespace HordeServer.VersionControl.Perforce
 							}
 						}
 
-						// Start any new tickers that don't exist yet
+						// Start any new tickers that don't exist yet, or update existing ones
 						foreach ((string clusterName, List<StreamInfo> streamInfos) in clusterToStreams)
 						{
 							if (!clusterToTicker.ContainsKey(clusterName))
@@ -326,6 +333,15 @@ namespace HordeServer.VersionControl.Perforce
 								ClusterTicker ticker = StartTicker(clusterName, streamInfos);
 								clusterToTicker.Add(ticker.ClusterName, ticker);
 							}
+							else
+							{
+								// Update existing ticker's stream list (includes refreshed Views from P4 stream changes)
+								if (!clusterToTicker[clusterName].Streams.Select(x => x.Hash).SequenceEqual(streamInfos.Select(x => x.Hash)))
+							{
+								_logger.LogInformation("Refreshing stream definitions for cluster {ClusterName}", clusterName);
+							}
+								clusterToTicker[clusterName].Streams = streamInfos;
+							}
 						}
 					}
 
@@ -554,7 +570,7 @@ namespace HordeServer.VersionControl.Perforce
 							{
 								bool hasAllFiles = describeRecord.Files.Count < MaxFiles;
 
-								ICommit commit = await CreateCommitAsync(perforce, streamInfo.StreamConfig, describeRecord, hasAllFiles, info, cancellationToken);
+								ICommit commit = await CreateCommitAsync(perforce, streamInfo.StreamConfig, describeRecord, hasAllFiles, info, cancellationToken, files.ToList());
 								_logger.LogInformation("Replicating {StreamId} commit {Change}", streamInfo.StreamConfig.Id, describeRecord.Number);
 
 								IReadOnlyList<CommitTag> commitTags = await commit.GetTagsAsync(cancellationToken);
