Horde: Fix perforce vcs stream info discovery
#rnx

diff --git a/Engine/Source/Programs/Horde/Plugins/Build/HordeServer.Build/VersionControl/Perforce/PerforceServiceCache.cs b/Engine/Source/Programs/Horde/Plugins/Build/HordeServer.Build/VersionControl/Perforce/PerforceServiceCache.cs
--- a/Engine/Source/Programs/Horde/Plugins/Build/HordeServer.Build/VersionControl/Perforce/PerforceServiceCache.cs
+++ b/Engine/Source/Programs/Horde/Plugins/Build/HordeServer.Build/VersionControl/Perforce/PerforceServiceCache.cs
@@ -364,7 +364,7 @@ namespace HordeServer.VersionControl.Perforce
 			IReadOnlyList<StreamConfig> streams = _buildConfig.CurrentValue.Streams;
 
 			Dictionary<string, List<StreamInfo>> clusters = new(StringComparer.OrdinalIgnoreCase);
-			foreach (IGrouping<string, StreamConfig> group in streams.GroupBy(x => x.ClusterName, StringComparer.OrdinalIgnoreCase))
+			foreach (IGrouping<string, StreamConfig> group in streams.Where(x => x.VCS.Equals("Perforce", StringComparison.OrdinalIgnoreCase)).GroupBy(x => x.ClusterName, StringComparer.OrdinalIgnoreCase))
 			{
 				try
 				{
