Horde: In the case where the instance type is not available in a particular AWS availability zone catch the exception and move onto the next instance type in the list
#rnx

diff --git a/Engine/Source/Programs/Horde/Plugins/Build/HordeServer.Build.Tests/Fleet/AwsRecyclingFleetManagerTest.cs b/Engine/Source/Programs/Horde/Plugins/Build/HordeServer.Build.Tests/Fleet/AwsRecyclingFleetManagerTest.cs
--- a/Engine/Source/Programs/Horde/Plugins/Build/HordeServer.Build.Tests/Fleet/AwsRecyclingFleetManagerTest.cs
+++ b/Engine/Source/Programs/Horde/Plugins/Build/HordeServer.Build.Tests/Fleet/AwsRecyclingFleetManagerTest.cs
@@ -175,6 +175,58 @@ namespace HordeServer.Tests.Fleet
 			Assert.AreEqual(FakeAmazonEc2.StatePending.Name, ec2.GetInstance(i4.InstanceId)!.State.Name);
 		}
 
+		[TestMethod]
+		public async Task ExpandOneAzUnsupportedOperationWithFallbackAsync()
+		{
+			FakeAmazonEc2 ec2 = new();
+			Instance i1 = ec2.AddInstance(FakeAmazonEc2.StateStopped, InstanceType.M5Large, FakeAmazonEc2.AzUsEast1A, tagPoolName: PoolName, tagPoolId: PoolId);
+			Instance i2 = ec2.AddInstance(FakeAmazonEc2.StateStopped, InstanceType.M5Large, FakeAmazonEc2.AzUsEast1A, tagPoolName: PoolName, tagPoolId: PoolId);
+ 
+			ec2.SetUnsupportedOperation(FakeAmazonEc2.AzUsEast1A, InstanceType.M5Large.Value);
+			ec2.SetCapacity(FakeAmazonEc2.AzUsEast1A, InstanceType.C5Large, 2);
+ 
+			await ExpandPoolAsync(ec2.Get(), 2, new(new List<string> { InstanceType.M5Large, InstanceType.C5Large }));
+ 
+			Assert.AreEqual(0, ec2.GetStoppedInstanceCount());
+			Assert.AreEqual(2, ec2.GetPendingInstanceCount());
+			Assert.AreEqual(InstanceType.C5Large, ec2.GetInstance(i1.InstanceId)!.InstanceType);
+			Assert.AreEqual(InstanceType.C5Large, ec2.GetInstance(i2.InstanceId)!.InstanceType);
+		}
+ 
+		[TestMethod]
+		public async Task ExpandUnsupportedOperationNoFallbackAsync()
+		{
+			FakeAmazonEc2 ec2 = new();
+			ec2.AddInstance(FakeAmazonEc2.StateStopped, InstanceType.M5Large, FakeAmazonEc2.AzUsEast1A, tagPoolName: PoolName);
+			ec2.SetUnsupportedOperation(FakeAmazonEc2.AzUsEast1A, InstanceType.M5Large.Value);
+ 
+			await ExpandPoolAsync(ec2.Get(), 1, new());
+ 
+			Assert.AreEqual(1, ec2.GetStoppedInstanceCount());
+			Assert.AreEqual(0, ec2.GetPendingInstanceCount());
+		}
+ 
+		[TestMethod]
+		public async Task ExpandMixedErrorsWithFallbackAsync()
+		{
+			FakeAmazonEc2 ec2 = new();
+			ec2.AddInstance(FakeAmazonEc2.StateStopped, InstanceType.M5Large, FakeAmazonEc2.AzUsEast1A, tagPoolName: PoolName);
+ 
+			// First: unsupported, Second: no capacity, Third: works
+			ec2.SetUnsupportedOperation(FakeAmazonEc2.AzUsEast1A, InstanceType.M5Large.Value);
+			ec2.SetCapacity(FakeAmazonEc2.AzUsEast1A, InstanceType.C5Large, 0);
+			ec2.SetCapacity(FakeAmazonEc2.AzUsEast1A, InstanceType.R5Large, 1);
+ 
+			await ExpandPoolAsync(ec2.Get(), 1, new(new List<string> {
+				InstanceType.M5Large,
+				InstanceType.C5Large,
+				InstanceType.R5Large
+			}));
+ 
+			Assert.AreEqual(0, ec2.GetStoppedInstanceCount());
+			Assert.AreEqual(1, ec2.GetPendingInstanceCount());
+		}
+		
 		[TestMethod]
 		public void DistributesInstanceRequestsEvenly()
 		{
diff --git a/Engine/Source/Programs/Horde/Plugins/Build/HordeServer.Build.Tests/Fleet/FakeAmazonEc2.cs b/Engine/Source/Programs/Horde/Plugins/Build/HordeServer.Build.Tests/Fleet/FakeAmazonEc2.cs
--- a/Engine/Source/Programs/Horde/Plugins/Build/HordeServer.Build.Tests/Fleet/FakeAmazonEc2.cs
+++ b/Engine/Source/Programs/Horde/Plugins/Build/HordeServer.Build.Tests/Fleet/FakeAmazonEc2.cs
@@ -43,6 +43,12 @@ public class FakeAmazonEc2
 		{AzUsEast1B, new () }
 	};
 
+	private readonly Dictionary<string, HashSet<string>> _unsupportedOperations = new()
+	{
+		{AzUsEast1A, new HashSet<string>()},
+		{AzUsEast1B, new HashSet<string>()}
+	};
+
 	public FakeAmazonEc2()
 	{
 		_mock = new Mock<IAmazonEC2>(MockBehavior.Strict);
@@ -75,6 +81,13 @@ public class FakeAmazonEc2
 		_availabilityZoneCapacity[availabilityZone] = temp;
 	}
 
+	public void SetUnsupportedOperation(string availabilityZone, string instanceType)
+	{
+		HashSet<string> unsupported = _unsupportedOperations.GetValueOrDefault(availabilityZone, new HashSet<string>());
+		unsupported.Add(instanceType);
+		_unsupportedOperations[availabilityZone] = unsupported;
+	}
+
 	public Instance AddInstance(
 		InstanceState state,
 		InstanceType type,
@@ -175,6 +188,16 @@ public class FakeAmazonEc2
 
 			string az = _instances[instanceId].Placement.AvailabilityZone;
 			string instanceType = _instances[instanceId].InstanceType;
+
+			// Check for unsupported operations first
+			if (_unsupportedOperations.GetValueOrDefault(az, new HashSet<string>()).Contains(instanceType))
+			{
+				throw new AmazonEC2Exception("The requested configuration is currently not supported.")
+				{
+					ErrorCode = "UnsupportedOperation"
+				};
+			}
+
 			int capacity = azInstanceTypeCapacities[az].GetValueOrDefault(instanceType, 0);
 
 			if (capacity <= 0)
diff --git a/Engine/Source/Programs/Horde/Plugins/Compute/HordeServer.Compute/Agents/Fleet/Providers/AwsRecyclingFleetManager.cs b/Engine/Source/Programs/Horde/Plugins/Compute/HordeServer.Compute/Agents/Fleet/Providers/AwsRecyclingFleetManager.cs
--- a/Engine/Source/Programs/Horde/Plugins/Compute/HordeServer.Compute/Agents/Fleet/Providers/AwsRecyclingFleetManager.cs
+++ b/Engine/Source/Programs/Horde/Plugins/Compute/HordeServer.Compute/Agents/Fleet/Providers/AwsRecyclingFleetManager.cs
@@ -269,6 +269,12 @@ public sealed class AwsRecyclingFleetManager : IFleetManager
 				_ec2StartInstanceCounter.Add(1, new[] { instanceTypeTag, azTag, statusTag });
 				_logger.LogInformation("Insufficient capacity for {InstanceType}", instanceType?.ToString());
 			}
+			catch (AmazonEC2Exception e) when (e.ErrorCode == "UnsupportedOperation")
+			{
+				KeyValuePair<string, object?> statusTag = KeyValuePair.Create("status", (object?)"errorUnsupported");
+				_ec2StartInstanceCounter.Add(1, new[] { instanceTypeTag, azTag, statusTag });
+				_logger.LogInformation("Unsupported operation for {InstanceType}", instanceType?.ToString());
+			}
 			catch (Exception)
 			{
 				KeyValuePair<string, object?> statusTag = KeyValuePair.Create("status", (object?)"error");
