Horde: Updating schema documentation

#rnx

diff --git a/Engine/Source/Programs/Horde/Docs/Config/Schema/AclActions.md b/Engine/Source/Programs/Horde/Docs/Config/Schema/AclActions.md
--- a/Engine/Source/Programs/Horde/Docs/Config/Schema/AclActions.md
+++ b/Engine/Source/Programs/Horde/Docs/Config/Schema/AclActions.md
@@ -37,6 +37,48 @@
 | `DeleteAccount` |  |
 | `ViewAccount` |  |
 
+## Artifacts
+
+| Name | Description |
+| ---- | ----------- |
+| `ReadArtifact` | Permission to read from an artifact |
+| `WriteArtifact` | Permission to write to an artifact |
+| `DeleteArtifact` | Permission to delete to an artifact |
+| `UploadArtifact` | Ability to create an artifact. Typically just for debugging; agents have this access for a particular session. |
+| `DownloadArtifact` | Ability to download an artifact |
+
+## Bisect
+
+| Name | Description |
+| ---- | ----------- |
+| `CreateBisectTask` | Ability to start new bisect tasks |
+| `UpdateBisectTask` | Ability to update a bisect task |
+| `ViewBisectTask` | Ability to view a bisect task |
+
+## BuildHealth
+
+| Name | Description |
+| ---- | ----------- |
+| `AddUpdateDeleteBuildHealthFilter` | Ability to add, update, and delete build health filters. |
+
+## Devices
+
+| Name | Description |
+| ---- | ----------- |
+| `DeviceRead` | Ability to read devices |
+| `DeviceWrite` | Ability to write devices |
+
+## Jobs
+
+| Name | Description |
+| ---- | ----------- |
+| `CreateJob` | Ability to start new jobs |
+| `UpdateJob` | Rename a job, modify its priority, etc... |
+| `DeleteJob` | Delete a job properties |
+| `ExecuteJob` | Allows updating a job metadata (name, changelist number, step properties, new groups, job states, etc...). Typically granted to agents. Not user facing. |
+| `RetryJobStep` | Ability to retry a failed job step |
+| `ViewJob` | Ability to view a job |
+
 ## Notices
 
 | Name | Description |
@@ -45,6 +87,28 @@
 | `UpdateNotice` |  |
 | `DeleteNotice` |  |
 
+## Notifications
+
+| Name | Description |
+| ---- | ----------- |
+| `CreateSubscription` | Ability to subscribe to notifications |
+
+## Projects
+
+| Name | Description |
+| ---- | ----------- |
+| `CreateProject` | Allows the creation of new projects |
+| `DeleteProject` | Allows deletion of projects. |
+| `UpdateProject` | Modify attributes of a project (name, categories, etc...) |
+| `ViewProject` | View information about a project |
+
+## Replicators
+
+| Name | Description |
+| ---- | ----------- |
+| `UpdateReplicator` | Allows deletion of projects. |
+| `ViewReplicator` | Allows the creation of new projects |
+
 ## ServiceAccounts
 
 | Name | Description |
@@ -54,6 +118,24 @@
 | `DeleteAccount` |  |
 | `ViewAccount` |  |
 
+## Streams
+
+| Name | Description |
+| ---- | ----------- |
+| `CreateStream` | Allows the creation of new streams within a project |
+| `UpdateStream` | Allows updating a stream (agent types, templates, schedules) |
+| `DeleteStream` | Allows deleting a stream |
+| `ViewStream` | Ability to view a stream |
+| `ViewChanges` | View changes submitted to a stream. NOTE: this returns responses from the server's Perforce account, which may be a priviledged user. |
+| `ViewTemplate` | View template associated with a stream |
+
+## TestData
+
+| Name | Description |
+| ---- | ----------- |
+| `TestHealthRead` | Ability to read test health |
+| `TestHealthWrite` | Ability to write test health |
+
 ## Accounts
 
 | Name | Description |
@@ -74,6 +156,30 @@
 | `ViewAgent` | View an agent |
 | `ListAgents` | List the available agents |
 
+## Artifacts
+
+| Name | Description |
+| ---- | ----------- |
+| `ReadArtifact` |  |
+| `WriteArtifact` |  |
+| `DeleteArtifact` |  |
+| `UploadArtifact` |  |
+| `DownloadArtifact` |  |
+
+## Bisect
+
+| Name | Description |
+| ---- | ----------- |
+| `CreateBisectTask` |  |
+| `UpdateBisectTask` |  |
+| `ViewBisectTask` |  |
+
+## BuildHealth
+
+| Name | Description |
+| ---- | ----------- |
+| `AddUpdateDeleteBuildHealthFilter` |  |
+
 ## Compute
 
 | Name | Description |
@@ -83,6 +189,24 @@
 | `UbaCacheRead` | User can read from Unreal Build Accelerator cache |
 | `UbaCacheWrite` | User can write to Unreal Build Accelerator cache |
 
+## Devices
+
+| Name | Description |
+| ---- | ----------- |
+| `DeviceRead` |  |
+| `DeviceWrite` |  |
+
+## Jobs
+
+| Name | Description |
+| ---- | ----------- |
+| `CreateJob` |  |
+| `UpdateJob` |  |
+| `DeleteJob` |  |
+| `ExecuteJob` |  |
+| `RetryJobStep` |  |
+| `ViewJob` |  |
+
 ## Leases
 
 | Name | Description |
@@ -109,6 +233,12 @@
 | `UpdateNotice` |  |
 | `DeleteNotice` |  |
 
+## Notifications
+
+| Name | Description |
+| ---- | ----------- |
+| `CreateSubscription` |  |
+
 ## Pools
 
 | Name | Description |
@@ -119,6 +249,22 @@
 | `ViewPool` | Ability to view a pool |
 | `ListPools` | View all the available agent pools |
 
+## Projects
+
+| Name | Description |
+| ---- | ----------- |
+| `CreateProject` |  |
+| `DeleteProject` |  |
+| `UpdateProject` |  |
+| `ViewProject` |  |
+
+## Replicators
+
+| Name | Description |
+| ---- | ----------- |
+| `UpdateReplicator` |  |
+| `ViewReplicator` |  |
+
 ## ServiceAccounts
 
 | Name | Description |
@@ -143,6 +289,24 @@
 | `DownloadSoftware` | Ability to download the agent software |
 | `DeleteSoftware` | Ability to delete agent software |
 
+## Streams
+
+| Name | Description |
+| ---- | ----------- |
+| `CreateStream` |  |
+| `UpdateStream` |  |
+| `DeleteStream` |  |
+| `ViewStream` |  |
+| `ViewChanges` |  |
+| `ViewTemplate` |  |
+
+## TestData
+
+| Name | Description |
+| ---- | ----------- |
+| `TestHealthRead` |  |
+| `TestHealthWrite` |  |
+
 ## Accounts
 
 | Name | Description |
@@ -163,6 +327,30 @@
 | `ViewAgent` |  |
 | `ListAgents` |  |
 
+## Artifacts
+
+| Name | Description |
+| ---- | ----------- |
+| `ReadArtifact` |  |
+| `WriteArtifact` |  |
+| `DeleteArtifact` |  |
+| `UploadArtifact` |  |
+| `DownloadArtifact` |  |
+
+## Bisect
+
+| Name | Description |
+| ---- | ----------- |
+| `CreateBisectTask` |  |
+| `UpdateBisectTask` |  |
+| `ViewBisectTask` |  |
+
+## BuildHealth
+
+| Name | Description |
+| ---- | ----------- |
+| `AddUpdateDeleteBuildHealthFilter` |  |
+
 ## Compute
 
 | Name | Description |
@@ -172,6 +360,24 @@
 | `UbaCacheRead` |  |
 | `UbaCacheWrite` |  |
 
+## Devices
+
+| Name | Description |
+| ---- | ----------- |
+| `DeviceRead` |  |
+| `DeviceWrite` |  |
+
+## Jobs
+
+| Name | Description |
+| ---- | ----------- |
+| `CreateJob` |  |
+| `UpdateJob` |  |
+| `DeleteJob` |  |
+| `ExecuteJob` |  |
+| `RetryJobStep` |  |
+| `ViewJob` |  |
+
 ## Leases
 
 | Name | Description |
@@ -198,6 +404,12 @@
 | `UpdateNotice` |  |
 | `DeleteNotice` |  |
 
+## Notifications
+
+| Name | Description |
+| ---- | ----------- |
+| `CreateSubscription` |  |
+
 ## Pools
 
 | Name | Description |
@@ -208,6 +420,22 @@
 | `ViewPool` |  |
 | `ListPools` |  |
 
+## Projects
+
+| Name | Description |
+| ---- | ----------- |
+| `CreateProject` |  |
+| `DeleteProject` |  |
+| `UpdateProject` |  |
+| `ViewProject` |  |
+
+## Replicators
+
+| Name | Description |
+| ---- | ----------- |
+| `UpdateReplicator` |  |
+| `ViewReplicator` |  |
+
 ## ServiceAccounts
 
 | Name | Description |
@@ -232,6 +460,24 @@
 | `DownloadSoftware` |  |
 | `DeleteSoftware` |  |
 
+## Streams
+
+| Name | Description |
+| ---- | ----------- |
+| `CreateStream` |  |
+| `UpdateStream` |  |
+| `DeleteStream` |  |
+| `ViewStream` |  |
+| `ViewChanges` |  |
+| `ViewTemplate` |  |
+
+## TestData
+
+| Name | Description |
+| ---- | ----------- |
+| `TestHealthRead` |  |
+| `TestHealthWrite` |  |
+
 ## Accounts
 
 | Name | Description |
@@ -252,5 +498,29 @@
 | `ViewAgent` |  |
 | `ListAgents` |  |
 
+## Artifacts
+
+| Name | Description |
+| ---- | ----------- |
+| `ReadArtifact` |  |
+| `WriteArtifact` |  |
+| `DeleteArtifact` |  |
+| `UploadArtifact` |  |
+| `DownloadArtifact` |  |
+
+## Bisect
+
+| Name | Description |
+| ---- | ----------- |
+| `CreateBisectTask` |  |
+| `UpdateBisectTask` |  |
+| `ViewBisectTask` |  |
+
+## BuildHealth
+
+| Name | Description |
+| ---- | ----------- |
+| `AddUpdateDeleteBuildHealthFilter` |  |
+
 ## Compute
 
@@ -261,6 +531,24 @@
 | `UbaCacheRead` |  |
 | `UbaCacheWrite` |  |
 
+## Devices
+
+| Name | Description |
+| ---- | ----------- |
+| `DeviceRead` |  |
+| `DeviceWrite` |  |
+
+## Jobs
+
+| Name | Description |
+| ---- | ----------- |
+| `CreateJob` |  |
+| `UpdateJob` |  |
+| `DeleteJob` |  |
+| `ExecuteJob` |  |
+| `RetryJobStep` |  |
+| `ViewJob` |  |
+
 ## Leases
 
 | Name | Description |
@@ -287,6 +575,12 @@
 | `UpdateNotice` |  |
 | `DeleteNotice` |  |
 
+## Notifications
+
+| Name | Description |
+| ---- | ----------- |
+| `CreateSubscription` |  |
+
 ## Pools
 
 | Name | Description |
@@ -297,11 +591,21 @@
 | `ViewPool` |  |
 | `ListPools` |  |
 
-## Secrets
+## Projects
 
 | Name | Description |
 | ---- | ----------- |
-| `ViewSecret` | View a credential |
+| `CreateProject` |  |
+| `DeleteProject` |  |
+| `UpdateProject` |  |
+| `ViewProject` |  |
 
+## Replicators
+
+| Name | Description |
+| ---- | ----------- |
+| `UpdateReplicator` |  |
+| `ViewReplicator` |  |
+
 ## ServiceAccounts
 
@@ -327,6 +631,24 @@
 | `DownloadSoftware` |  |
 | `DeleteSoftware` |  |
 
+## Streams
+
+| Name | Description |
+| ---- | ----------- |
+| `CreateStream` |  |
+| `UpdateStream` |  |
+| `DeleteStream` |  |
+| `ViewStream` |  |
+| `ViewChanges` |  |
+| `ViewTemplate` |  |
+
+## TestData
+
+| Name | Description |
+| ---- | ----------- |
+| `TestHealthRead` |  |
+| `TestHealthWrite` |  |
+
 ## Accounts
 
 | Name | Description |
@@ -347,6 +669,30 @@
 | `ViewAgent` |  |
 | `ListAgents` |  |
 
+## Artifacts
+
+| Name | Description |
+| ---- | ----------- |
+| `ReadArtifact` |  |
+| `WriteArtifact` |  |
+| `DeleteArtifact` |  |
+| `UploadArtifact` |  |
+| `DownloadArtifact` |  |
+
+## Bisect
+
+| Name | Description |
+| ---- | ----------- |
+| `CreateBisectTask` |  |
+| `UpdateBisectTask` |  |
+| `ViewBisectTask` |  |
+
+## BuildHealth
+
+| Name | Description |
+| ---- | ----------- |
+| `AddUpdateDeleteBuildHealthFilter` |  |
+
 ## Compute
 
 | Name | Description |
@@ -356,6 +702,24 @@
 | `UbaCacheRead` |  |
 | `UbaCacheWrite` |  |
 
+## Devices
+
+| Name | Description |
+| ---- | ----------- |
+| `DeviceRead` |  |
+| `DeviceWrite` |  |
+
+## Jobs
+
+| Name | Description |
+| ---- | ----------- |
+| `CreateJob` |  |
+| `UpdateJob` |  |
+| `DeleteJob` |  |
+| `ExecuteJob` |  |
+| `RetryJobStep` |  |
+| `ViewJob` |  |
+
 ## Leases
 
 | Name | Description |
@@ -382,6 +746,12 @@
 | `UpdateNotice` |  |
 | `DeleteNotice` |  |
 
+## Notifications
+
+| Name | Description |
+| ---- | ----------- |
+| `CreateSubscription` |  |
+
 ## Pools
 
 | Name | Description |
@@ -392,11 +762,27 @@
 | `ViewPool` |  |
 | `ListPools` |  |
 
+## Projects
+
+| Name | Description |
+| ---- | ----------- |
+| `CreateProject` |  |
+| `DeleteProject` |  |
+| `UpdateProject` |  |
+| `ViewProject` |  |
+
+## Replicators
+
+| Name | Description |
+| ---- | ----------- |
+| `UpdateReplicator` |  |
+| `ViewReplicator` |  |
+
 ## Secrets
 
 | Name | Description |
 | ---- | ----------- |
-| `ViewSecret` |  |
+| `ViewSecret` | View a credential |
 
 ## ServiceAccounts
 
@@ -422,18 +808,23 @@
 | `DownloadSoftware` |  |
 | `DeleteSoftware` |  |
 
-## Storage
+## Streams
 
 | Name | Description |
 | ---- | ----------- |
-| `ReadBlobs` | Ability to read blobs from the storage service |
-| `WriteBlobs` | Ability to write blobs to the storage service |
-| `ReadAliases` | Ability to read aliases from the storage service |
-| `WriteAliases` | Ability to write aliases to the storage service |
-| `DeleteAliases` | Ability to write aliases to the storage service |
-| `ReadRefs` | Ability to read refs from the storage service |
-| `WriteRefs` | Ability to write refs to the storage service |
-| `DeleteRefs` | Ability to delete refs |
+| `CreateStream` |  |
+| `UpdateStream` |  |
+| `DeleteStream` |  |
+| `ViewStream` |  |
+| `ViewChanges` |  |
+| `ViewTemplate` |  |
 
+## TestData
+
+| Name | Description |
+| ---- | ----------- |
+| `TestHealthRead` |  |
+| `TestHealthWrite` |  |
+
 ## Accounts
 
@@ -459,25 +850,25 @@
 
 | Name | Description |
 | ---- | ----------- |
-| `ReadArtifact` | Permission to read from an artifact |
-| `WriteArtifact` | Permission to write to an artifact |
-| `DeleteArtifact` | Permission to delete to an artifact |
-| `UploadArtifact` | Ability to create an artifact. Typically just for debugging; agents have this access for a particular session. |
-| `DownloadArtifact` | Ability to download an artifact |
+| `ReadArtifact` |  |
+| `WriteArtifact` |  |
+| `DeleteArtifact` |  |
+| `UploadArtifact` |  |
+| `DownloadArtifact` |  |
 
 ## Bisect
 
 | Name | Description |
 | ---- | ----------- |
-| `CreateBisectTask` | Ability to start new bisect tasks |
-| `UpdateBisectTask` | Ability to update a bisect task |
-| `ViewBisectTask` | Ability to view a bisect task |
+| `CreateBisectTask` |  |
+| `UpdateBisectTask` |  |
+| `ViewBisectTask` |  |
 
 ## BuildHealth
 
 | Name | Description |
 | ---- | ----------- |
-| `AddUpdateDeleteBuildHealthFilter` | Ability to add, update, and delete build health filters. |
+| `AddUpdateDeleteBuildHealthFilter` |  |
 
 ## Compute
 
@@ -492,19 +883,19 @@
 
 | Name | Description |
 | ---- | ----------- |
-| `DeviceRead` | Ability to read devices |
-| `DeviceWrite` | Ability to write devices |
+| `DeviceRead` |  |
+| `DeviceWrite` |  |
 
 ## Jobs
 
 | Name | Description |
 | ---- | ----------- |
-| `CreateJob` | Ability to start new jobs |
-| `UpdateJob` | Rename a job, modify its priority, etc... |
-| `DeleteJob` | Delete a job properties |
-| `ExecuteJob` | Allows updating a job metadata (name, changelist number, step properties, new groups, job states, etc...). Typically granted to agents. Not user facing. |
-| `RetryJobStep` | Ability to retry a failed job step |
-| `ViewJob` | Ability to view a job |
+| `CreateJob` |  |
+| `UpdateJob` |  |
+| `DeleteJob` |  |
+| `ExecuteJob` |  |
+| `RetryJobStep` |  |
+| `ViewJob` |  |
 
 ## Leases
 
@@ -536,7 +927,7 @@
 
 | Name | Description |
 | ---- | ----------- |
-| `CreateSubscription` | Ability to subscribe to notifications |
+| `CreateSubscription` |  |
 
 ## Pools
 
@@ -552,17 +943,17 @@
 
 | Name | Description |
 | ---- | ----------- |
-| `CreateProject` | Allows the creation of new projects |
-| `DeleteProject` | Allows deletion of projects. |
-| `UpdateProject` | Modify attributes of a project (name, categories, etc...) |
-| `ViewProject` | View information about a project |
+| `CreateProject` |  |
+| `DeleteProject` |  |
+| `UpdateProject` |  |
+| `ViewProject` |  |
 
 ## Replicators
 
 | Name | Description |
 | ---- | ----------- |
-| `UpdateReplicator` | Allows deletion of projects. |
-| `ViewReplicator` | Allows the creation of new projects |
+| `UpdateReplicator` |  |
+| `ViewReplicator` |  |
 
 ## Secrets
 
@@ -594,36 +985,23 @@
 | `DownloadSoftware` |  |
 | `DeleteSoftware` |  |
 
-## Storage
-
-| Name | Description |
-| ---- | ----------- |
-| `ReadBlobs` |  |
-| `WriteBlobs` |  |
-| `ReadAliases` |  |
-| `WriteAliases` |  |
-| `DeleteAliases` |  |
-| `ReadRefs` |  |
-| `WriteRefs` |  |
-| `DeleteRefs` |  |
-
 ## Streams
 
 | Name | Description |
 | ---- | ----------- |
-| `CreateStream` | Allows the creation of new streams within a project |
-| `UpdateStream` | Allows updating a stream (agent types, templates, schedules) |
-| `DeleteStream` | Allows deleting a stream |
-| `ViewStream` | Ability to view a stream |
-| `ViewChanges` | View changes submitted to a stream. NOTE: this returns responses from the server's Perforce account, which may be a priviledged user. |
-| `ViewTemplate` | View template associated with a stream |
+| `CreateStream` |  |
+| `UpdateStream` |  |
+| `DeleteStream` |  |
+| `ViewStream` |  |
+| `ViewChanges` |  |
+| `ViewTemplate` |  |
 
 ## TestData
 
 | Name | Description |
 | ---- | ----------- |
-| `TestHealthRead` | Ability to read test health |
-| `TestHealthWrite` | Ability to write test health |
+| `TestHealthRead` |  |
+| `TestHealthWrite` |  |
 
 ## Accounts
 
@@ -788,14 +1166,14 @@
 
 | Name | Description |
 | ---- | ----------- |
-| `ReadBlobs` |  |
-| `WriteBlobs` |  |
-| `ReadAliases` |  |
-| `WriteAliases` |  |
-| `DeleteAliases` |  |
-| `ReadRefs` |  |
-| `WriteRefs` |  |
-| `DeleteRefs` |  |
+| `ReadBlobs` | Ability to read blobs from the storage service |
+| `WriteBlobs` | Ability to write blobs to the storage service |
+| `ReadAliases` | Ability to read aliases from the storage service |
+| `WriteAliases` | Ability to write aliases to the storage service |
+| `DeleteAliases` | Ability to write aliases to the storage service |
+| `ReadRefs` | Ability to read refs from the storage service |
+| `WriteRefs` | Ability to write refs to the storage service |
+| `DeleteRefs` | Ability to delete refs |
 
 ## Streams
 
@@ -808,12 +1186,6 @@
 | `ViewChanges` |  |
 | `ViewTemplate` |  |
 
-## Symbols
-
-| Name | Description |
-| ---- | ----------- |
-| `ReadSymbols` | Ability to download symbols |
-
 ## TestData
 
 | Name | Description |
@@ -1008,7 +1380,7 @@
 
 | Name | Description |
 | ---- | ----------- |
-| `ReadSymbols` |  |
+| `ReadSymbols` | Ability to download symbols |
 
 ## TestData
 
diff --git a/Engine/Source/Programs/Horde/Docs/Config/Schema/Globals.md b/Engine/Source/Programs/Horde/Docs/Config/Schema/Globals.md
--- a/Engine/Source/Programs/Horde/Docs/Config/Schema/Globals.md
+++ b/Engine/Source/Programs/Horde/Docs/Config/Schema/Globals.md
@@ -59,38 +59,149 @@
 
 Name | Description
 ---- | -----------
-`analytics` | [AnalyticsConfig](#analyticsconfig)<br>Configuration for the analytics plugin
+`build` | [BuildConfig](#buildconfig)<br>Configuration for the build plugin
 `compute` | [ComputeConfig](#computeconfig)<br>Configuration for the compute plugin
-`experimental` | [ExperimentalConfig](#experimentalconfig)<br>Configuration for the experimental plugin
+`analytics` | [AnalyticsConfig](#analyticsconfig)<br>Configuration for the analytics plugin
+`exampleplugin` | [ExamplePluginConfig](#examplepluginconfig)<br>Configuration for the exampleplugin plugin
 `secrets` | [SecretsConfig](#secretsconfig)<br>Configuration for the secrets plugin
+`experimental` | [ExperimentalConfig](#experimentalconfig)<br>Configuration for the experimental plugin
 `storage` | [StorageConfig](#storageconfig)<br>Configuration for the storage plugin
-`build` | [BuildConfig](#buildconfig)<br>Configuration for the build plugin
 `symbols` | [SymbolsConfig](#symbolsconfig)<br>Configuration for the symbols plugin
-`exampleplugin` | [ExamplePluginConfig](#examplepluginconfig)<br>Configuration for the exampleplugin plugin
 `tools` | [ToolsConfig](#toolsconfig)<br>Configuration for the tools plugin
 
-## AnalyticsConfig
+## BuildConfig
 
-Config settings for analytics
+Configuration for the build plugin
 
 Name | Description
 ---- | -----------
-`stores` | [TelemetryStoreConfig](Telemetry.md)`[]`<br>Metrics to aggregate on the Horde server
+`perforceClusters` | [PerforceCluster](#perforcecluster)`[]`<br>List of Perforce clusters
+`devices` | [DeviceConfig](#deviceconfig)<br>Device configuration
+`maxConformCount` | `integer`<br>Maximum number of conforms to run at once
+`agentShutdownIfDisabledGracePeriod` | `string`<br>Time to wait before shutting down an agent that has been disabled Used if no value is set on the actual pool.
+`artifactTypes` | [ArtifactTypeConfig](#artifacttypeconfig)`[]`<br>Configuration for different artifact types
+`projects` | [ProjectConfig](Projects.md)`[]`<br>List of projects
+`enableConformTasks` | `boolean`<br>Whether to allow conform tasks to run
+`enableJobTaskConformCheck` | `boolean`<br>Whether to check if a conform is required on an agent before assigning a lease for a job task source.
+`issueFixedTag` | `string`<br>Commit tag to use for marking issues as fixed
+`issues` | [IssueConfig](Issues.md)<br>Build health issue configurations
+`enableIssueRootCauseAnalysisNotifications` | `boolean`<br>Whether to allow notifications for Root Cause Analysis to be sent
+`autoRcaUseMultiStage` | `boolean`<br>Whether to use multi-stage LLM analysis for auto-RCA. Multi-stage is slower but potentially more accurate. Set to false for faster processing with higher throughput.
+`maxConcurrentRcaRequests` | `integer`<br>Maximum number of concurrent auto-RCA requests. Limits load on the external LLM API. Set to 0 to disable auto-RCA globally (acts as a kill switch).
+`issueFindSuspectsMaxSuccessDays` | `integer`<br>The maximum days since node success to scan for suspects
+`issueFindSuspectsMaxScoringSeconds` | `integer`<br>The maximum number of seconds to spend scording suspects
+`issueFindSuspectsMaxChanges` | `integer`<br>Maximum number of commits to examine for suspects. Previously hardcoded as const MaxChanges = 500.
+`issueFindSuspectsEnableParallelScoring` | `boolean`<br>Enable parallel scoring of commits for suspect finding. When enabled, commits are scored concurrently up to MaxParallelism.
+`issueFindSuspectsMaxParallelism` | `integer`<br>Maximum concurrent commit file lookups when parallel scoring is enabled. Higher values improve speed but increase Perforce server load.
+`issueFindSuspectsEnableEarlyExit` | `boolean`<br>Enable early exit when theoretical maximum score is achieved. Only exits early if the best possible score has been found.
+`issueFindSuspectsEnableBidirectionalQuery` | `boolean`<br>Enable bidirectional commit query to capture both recent and old commits. Queries half the limit from each end of the commit range.
+`issueFindSuspectsEnableFingerprintFallback` | `boolean`<br>Enable fingerprint-based fallback when no suspects are found. Helps prevent duplicate issues when commit range exceeds MaxChanges.
+`issueFingerprintFallbackMinScore` | `integer`<br>Minimum score required for fingerprint fallback matching. Higher values are more conservative (fewer false matches).
+`issueFingerprintFallbackMaxAgeDays` | `integer`<br>Maximum age in days for issues to be considered in fingerprint fallback.
+`issueFingerprintFallbackExcludedTypes` | `string[]`<br>Fingerprint types excluded from fallback matching (too generic).
+`enableExpandedIgnorePatterns` | `boolean`<br>Whether to enable expanded IgnorePatterns discovery (workspace-wide scanning + config-specified files). When false (default), only Engine directory is scanned for ignore patterns.
 
-## ComputeConfig
+## PerforceCluster
 
-Configuration for the compute system
+Information about a cluster of Perforce servers.
 
 Name | Description
 ---- | -----------
-`acl` | [AclConfig](#aclconfig)<br>Inherited root acl
-`versionEnum` | [ConfigVersion](#configversion-enum)<br>Config version number
-`telemetryStoreId` | `string`<br>Telemetry store for lease and compute telemetry data
-`rates` | [AgentRateConfig](#agentrateconfig)`[]`<br>List of costs of a particular agent type
-`clusters` | [ComputeClusterConfig](#computeclusterconfig)`[]`<br>List of compute profiles
-`pools` | [PoolConfig](#poolconfig)`[]`<br>List of pools
-`software` | [AgentSoftwareConfig](#agentsoftwareconfig)`[]`<br>List of costs of a particular agent type
-`networks` | [NetworkConfig](#networkconfig)`[]`<br>List of networks
+`name` | `string`<br>Name of the cluster
+`serviceAccount` | `string`<br>Username for Horde to log in to this server. Will use the first account specified below if not overridden.
+`canImpersonate` | `boolean`<br>Whether the service account can impersonate other users
+`supportsPartitionedWorkspaces` | `boolean`<br>Whether to use partitioned workspaces on this server
+`servers` | [PerforceServer](#perforceserver)`[]`<br>List of servers
+`credentials` | [PerforceCredentials](#perforcecredentials)`[]`<br>List of server credentials
+`autoSdk` | [AutoSdkWorkspace](#autosdkworkspace)`[]`<br>List of autosdk streams
 
+## PerforceServer
+
+Information about an individual Perforce server
+
+Name | Description
+---- | -----------
+`serverAndPort` | `string`<br>The server and port. The server may be a DNS entry with multiple records, in which case it will be actively load balanced. If "ssl:" prefix is used, ensure P4 server's fingerprint/certificate is trusted. See Horde's documentation on connecting to SSL-enabled Perforce servers.
+`healthCheck` | `boolean`<br>Whether to query the healthcheck address under each server
+`resolveDns` | `boolean`<br>Whether to resolve the DNS entries and load balance between different hosts
+`maxConformCount` | `integer`<br>Maximum number of simultaneous conforms on this server
+`condition` | `string`<br>Optional condition for a machine to be eligible to use this server
+`properties` | `string[]`<br>List of properties for an agent to be eligible to use this server
+
+## PerforceCredentials
+
+Credentials for a Perforce user
+
+Name | Description
+---- | -----------
+`userName` | `string`<br>The username
+`password` | `string`<br>Password for the user
+`ticket` | `string`<br>Login ticket for the user (will be used instead of password if set)
+
+## AutoSdkWorkspace
+
+Path to a platform and stream to use for syncing AutoSDK
+
+Name | Description
+---- | -----------
+`name` | `string`<br>Name of this workspace
+`properties` | `string[]`<br>The agent properties to check (eg. "OSFamily=Windows")
+`userName` | `string`<br>Username for logging in to the server
+`stream` | `string`<br>Stream to use
+
+## DeviceConfig
+
+Configuration for devices
+
+Name | Description
+---- | -----------
+`platforms` | [DevicePlatformConfig](#deviceplatformconfig)`[]`<br>List of device platforms
+`pools` | [DevicePoolConfig](#devicepoolconfig)`[]`<br>List of device pools
+
+## DevicePlatformConfig
+
+Configuration for a device platform
+
+Name | Description
+---- | -----------
+`id` | `string`<br>The id for this platform
+`name` | `string`<br>Name of the platform
+`models` | `string[]`<br>A list of platform models
+`legacyNames` | `string[]`<br>Legacy names which older versions of Gauntlet may be using
+`legacyPerfSpecHighModel` | `string`<br>Model name for the high perf spec, which may be requested by Gauntlet
+
+## DevicePoolConfig
+
+Configuration for a device pool
+
+Name | Description
+---- | -----------
+`id` | `string`<br>The id for this platform
+`name` | `string`<br>The name of the pool
+`poolType` | [DevicePoolType](#devicepooltype-enum)<br>The type of the pool
+`projectIds` | `string[]`<br>List of project ids associated with pool
+
+## DevicePoolType (Enum)
+
+The type of device pool
+
+Name | Description
+---- | -----------
+`Automation` | Available to CIS jobs
+`Shared` | Shared by users with remote checking and checkouts
+
+## ArtifactTypeConfig
+
+Configuration for an artifact
+
+Name | Description
+---- | -----------
+`name` | `string`<br>Legacy 'Name' property
+`type` | `string`<br>Name of the artifact type
+`acl` | [AclConfig](#aclconfig)<br>Acl for the artifact type
+`keepCount` | `integer`<br>Number of artifacts to retain
+`keepDays` | `integer`<br>Number of days to retain artifacts of this type
+`namespaceId` | `string`<br>Storage namespace to use for this artifact types
+
 ## AclConfig
 
@@ -133,6 +244,21 @@
 `excludeActions` | `string[]`<br>Actions to exclude from the inherited actions
 `extends` | `string[]`<br>Other profiles to extend from
 
+## ComputeConfig
+
+Configuration for the compute system
+
+Name | Description
+---- | -----------
+`acl` | [AclConfig](#aclconfig)<br>Inherited root acl
+`versionEnum` | [ConfigVersion](#configversion-enum)<br>Config version number
+`telemetryStoreId` | `string`<br>Telemetry store for lease and compute telemetry data
+`rates` | [AgentRateConfig](#agentrateconfig)`[]`<br>List of costs of a particular agent type
+`clusters` | [ComputeClusterConfig](#computeclusterconfig)`[]`<br>List of compute profiles
+`pools` | [PoolConfig](#poolconfig)`[]`<br>List of pools
+`software` | [AgentSoftwareConfig](#agentsoftwareconfig)`[]`<br>List of costs of a particular agent type
+`networks` | [NetworkConfig](#networkconfig)`[]`<br>List of networks
+
 ## ConfigVersion (Enum)
 
 Global version number for running the server. As new features are introduced that require data migrations, this version number indicates the backwards compatibility functionality that must be enabled. When adding a new version here, also add a message to ConfigService.CreateSnapshotAsync describing the steps that need to be taken to upgrade the deployment.
@@ -315,13 +441,21 @@
 `description` | `string`<br>Human-readable description
 `computeId` | `string`<br>Compute ID for this network (used when allocating compute resources)
 
-## ExperimentalConfig
+## AnalyticsConfig
 
-Configuration for the experimental plugin
+Config settings for analytics
 
 Name | Description
 ---- | -----------
-`notifications` | [NotificationConfig](Notification.md)`[]`<br>Notifications to be sent on the Horde server
+`stores` | [TelemetryStoreConfig](Telemetry.md)`[]`<br>Metrics to aggregate on the Horde server
 
+## ExamplePluginConfig
+
+Configuration for the ExamplePlugin plugin
+
+Name | Description
+---- | -----------
+`greetingMessage` | `string`<br>The Greeting provided by the example plugin
+
 ## SecretsConfig
 
@@ -405,6 +539,14 @@
 ---- | -----------
 `vaultUri` | `string`<br>The Uri to the Azure Key Vault.
 
+## ExperimentalConfig
+
+Configuration for the experimental plugin
+
+Name | Description
+---- | -----------
+`notifications` | [NotificationConfig](Notification.md)`[]`<br>Notifications to be sent on the Horde server
+
 ## StorageConfig
 
 Configuration for storage
@@ -477,140 +619,6 @@
 `enableAliases` | `boolean`<br>Support querying exports by their aliases
 `acl` | [AclConfig](#aclconfig)<br>Access list for this namespace
 
-## BuildConfig
-
-Configuration for the build plugin
-
-Name | Description
----- | -----------
-`perforceClusters` | [PerforceCluster](#perforcecluster)`[]`<br>List of Perforce clusters
-`devices` | [DeviceConfig](#deviceconfig)<br>Device configuration
-`maxConformCount` | `integer`<br>Maximum number of conforms to run at once
-`agentShutdownIfDisabledGracePeriod` | `string`<br>Time to wait before shutting down an agent that has been disabled Used if no value is set on the actual pool.
-`artifactTypes` | [ArtifactTypeConfig](#artifacttypeconfig)`[]`<br>Configuration for different artifact types
-`projects` | [ProjectConfig](Projects.md)`[]`<br>List of projects
-`enableConformTasks` | `boolean`<br>Whether to allow conform tasks to run
-`enableJobTaskConformCheck` | `boolean`<br>Whether to check if a conform is required on an agent before assigning a lease for a job task source.
-`issueFixedTag` | `string`<br>Commit tag to use for marking issues as fixed
-`issues` | [IssueConfig](Issues.md)<br>Build health issue configurations
-`enableIssueRootCauseAnalysisNotifications` | `boolean`<br>Whether to allow notifications for Root Cause Analysis to be sent
-`autoRcaUseMultiStage` | `boolean`<br>Whether to use multi-stage LLM analysis for auto-RCA. Multi-stage is slower but potentially more accurate. Set to false for faster processing with higher throughput.
-`maxConcurrentRcaRequests` | `integer`<br>Maximum number of concurrent auto-RCA requests. Limits load on the external LLM API. Set to 0 to disable auto-RCA globally (acts as a kill switch).
-`issueFindSuspectsMaxSuccessDays` | `integer`<br>The maximum days since node success to scan for suspects
-`issueFindSuspectsMaxScoringSeconds` | `integer`<br>The maximum number of seconds to spend scording suspects
-`issueFindSuspectsMaxChanges` | `integer`<br>Maximum number of commits to examine for suspects. Previously hardcoded as const MaxChanges = 500.
-`issueFindSuspectsEnableParallelScoring` | `boolean`<br>Enable parallel scoring of commits for suspect finding. When enabled, commits are scored concurrently up to MaxParallelism.
-`issueFindSuspectsMaxParallelism` | `integer`<br>Maximum concurrent commit file lookups when parallel scoring is enabled. Higher values improve speed but increase Perforce server load.
-`issueFindSuspectsEnableEarlyExit` | `boolean`<br>Enable early exit when theoretical maximum score is achieved. Only exits early if the best possible score has been found.
-`issueFindSuspectsEnableBidirectionalQuery` | `boolean`<br>Enable bidirectional commit query to capture both recent and old commits. Queries half the limit from each end of the commit range.
-`issueFindSuspectsEnableFingerprintFallback` | `boolean`<br>Enable fingerprint-based fallback when no suspects are found. Helps prevent duplicate issues when commit range exceeds MaxChanges.
-`issueFingerprintFallbackMinScore` | `integer`<br>Minimum score required for fingerprint fallback matching. Higher values are more conservative (fewer false matches).
-`issueFingerprintFallbackMaxAgeDays` | `integer`<br>Maximum age in days for issues to be considered in fingerprint fallback.
-`issueFingerprintFallbackExcludedTypes` | `string[]`<br>Fingerprint types excluded from fallback matching (too generic).
-`enableExpandedIgnorePatterns` | `boolean`<br>Whether to enable expanded IgnorePatterns discovery (workspace-wide scanning + config-specified files). When false (default), only Engine directory is scanned for ignore patterns.
-
-## PerforceCluster
-
-Information about a cluster of Perforce servers.
-
-Name | Description
----- | -----------
-`name` | `string`<br>Name of the cluster
-`serviceAccount` | `string`<br>Username for Horde to log in to this server. Will use the first account specified below if not overridden.
-`canImpersonate` | `boolean`<br>Whether the service account can impersonate other users
-`supportsPartitionedWorkspaces` | `boolean`<br>Whether to use partitioned workspaces on this server
-`servers` | [PerforceServer](#perforceserver)`[]`<br>List of servers
-`credentials` | [PerforceCredentials](#perforcecredentials)`[]`<br>List of server credentials
-`autoSdk` | [AutoSdkWorkspace](#autosdkworkspace)`[]`<br>List of autosdk streams
-
-## PerforceServer
-
-Information about an individual Perforce server
-
-Name | Description
----- | -----------
-`serverAndPort` | `string`<br>The server and port. The server may be a DNS entry with multiple records, in which case it will be actively load balanced. If "ssl:" prefix is used, ensure P4 server's fingerprint/certificate is trusted. See Horde's documentation on connecting to SSL-enabled Perforce servers.
-`healthCheck` | `boolean`<br>Whether to query the healthcheck address under each server
-`resolveDns` | `boolean`<br>Whether to resolve the DNS entries and load balance between different hosts
-`maxConformCount` | `integer`<br>Maximum number of simultaneous conforms on this server
-`condition` | `string`<br>Optional condition for a machine to be eligible to use this server
-`properties` | `string[]`<br>List of properties for an agent to be eligible to use this server
-
-## PerforceCredentials
-
-Credentials for a Perforce user
-
-Name | Description
----- | -----------
-`userName` | `string`<br>The username
-`password` | `string`<br>Password for the user
-`ticket` | `string`<br>Login ticket for the user (will be used instead of password if set)
-
-## AutoSdkWorkspace
-
-Path to a platform and stream to use for syncing AutoSDK
-
-Name | Description
----- | -----------
-`name` | `string`<br>Name of this workspace
-`properties` | `string[]`<br>The agent properties to check (eg. "OSFamily=Windows")
-`userName` | `string`<br>Username for logging in to the server
-`stream` | `string`<br>Stream to use
-
-## DeviceConfig
-
-Configuration for devices
-
-Name | Description
----- | -----------
-`platforms` | [DevicePlatformConfig](#deviceplatformconfig)`[]`<br>List of device platforms
-`pools` | [DevicePoolConfig](#devicepoolconfig)`[]`<br>List of device pools
-
-## DevicePlatformConfig
-
-Configuration for a device platform
-
-Name | Description
----- | -----------
-`id` | `string`<br>The id for this platform
-`name` | `string`<br>Name of the platform
-`models` | `string[]`<br>A list of platform models
-`legacyNames` | `string[]`<br>Legacy names which older versions of Gauntlet may be using
-`legacyPerfSpecHighModel` | `string`<br>Model name for the high perf spec, which may be requested by Gauntlet
-
-## DevicePoolConfig
-
-Configuration for a device pool
-
-Name | Description
----- | -----------
-`id` | `string`<br>The id for this platform
-`name` | `string`<br>The name of the pool
-`poolType` | [DevicePoolType](#devicepooltype-enum)<br>The type of the pool
-`projectIds` | `string[]`<br>List of project ids associated with pool
-
-## DevicePoolType (Enum)
-
-The type of device pool
-
-Name | Description
----- | -----------
-`Automation` | Available to CIS jobs
-`Shared` | Shared by users with remote checking and checkouts
-
-## ArtifactTypeConfig
-
-Configuration for an artifact
-
-Name | Description
----- | -----------
-`name` | `string`<br>Legacy 'Name' property
-`type` | `string`<br>Name of the artifact type
-`acl` | [AclConfig](#aclconfig)<br>Acl for the artifact type
-`keepCount` | `integer`<br>Number of artifacts to retain
-`keepDays` | `integer`<br>Number of days to retain artifacts of this type
-`namespaceId` | `string`<br>Storage namespace to use for this artifact types
-
 ## SymbolsConfig
 
 Configuration for the tools system
@@ -630,14 +638,6 @@
 `public` | `boolean`<br>Whether to make this store available without auth
 `acl` | [AclConfig](#aclconfig)<br>Access to the symbol store
 
-## ExamplePluginConfig
-
-Configuration for the ExamplePlugin plugin
-
-Name | Description
----- | -----------
-`greetingMessage` | `string`<br>The Greeting provided by the example plugin
-
 ## ToolsConfig
 
 Configuration for the tools system
diff --git a/Engine/Source/Programs/Horde/Docs/Deployment/ServerSettings.md b/Engine/Source/Programs/Horde/Docs/Deployment/ServerSettings.md
--- a/Engine/Source/Programs/Horde/Docs/Deployment/ServerSettings.md
+++ b/Engine/Source/Programs/Horde/Docs/Deployment/ServerSettings.md
@@ -127,174 +127,16 @@
 
 Name | Description
 ---- | -----------
-`analytics` | [AnalyticsServerConfig](#analyticsserverconfig)<br>Configuration for the analytics plugin
+`build` | [BuildServerConfig](#buildserverconfig)<br>Configuration for the build plugin
 `compute` | [ComputeServerConfig](#computeserverconfig)<br>Configuration for the compute plugin
-`experimental` | [ExperimentalServerConfig](#experimentalserverconfig)<br>Configuration for the experimental plugin
+`analytics` | [AnalyticsServerConfig](#analyticsserverconfig)<br>Configuration for the analytics plugin
+`exampleplugin` | [ExamplePluginServerConfig](#examplepluginserverconfig)<br>Configuration for the exampleplugin plugin
 `secrets` | [SecretsServerConfig](#secretsserverconfig)<br>Configuration for the secrets plugin
+`experimental` | [ExperimentalServerConfig](#experimentalserverconfig)<br>Configuration for the experimental plugin
 `storage` | [StorageServerConfig](#storageserverconfig)<br>Configuration for the storage plugin
-`build` | [BuildServerConfig](#buildserverconfig)<br>Configuration for the build plugin
 `symbols` | [PluginServerConfig](#pluginserverconfig)<br>Configuration for the symbols plugin
-`exampleplugin` | [ExamplePluginServerConfig](#examplepluginserverconfig)<br>Configuration for the exampleplugin plugin
 `tools` | [ToolsServerConfig](#toolsserverconfig)<br>Configuration for the tools plugin
 
-## AnalyticsServerConfig
-
-Server configuration for the analytics system
-
-Name | Description
----- | -----------
-`sinks` | [TelemetrySinkConfig](#telemetrysinkconfig)<br>Settings for the various telemetry sinks
-`enabled` | `boolean`<br>Whether the plugin should be enabled or not
-
-## TelemetrySinkConfig
-
-Telemetry sinks
-
-Name | Description
----- | -----------
-`epic` | [EpicTelemetryConfig](#epictelemetryconfig)<br>Settings for the Epic telemetry sink
-`mongo` | [MongoTelemetryConfig](#mongotelemetryconfig)<br>Settings for the MongoDB telemetry sink
-
-## EpicTelemetryConfig
-
-Configuration for the telemetry sink
-
-Name | Description
----- | -----------
-`url` | `string`<br>Base URL for the telemetry server
-`appId` | `string`<br>Application name to send in the event messages
-`enabled` | `boolean`<br>Whether to enable this sink
-
-## MongoTelemetryConfig
-
-Configuration for the telemetry sink
-
-Name | Description
----- | -----------
-`retainDays` | `number`<br>Number of days worth of telmetry events to keep
-`enabled` | `boolean`<br>Whether to enable this sink
-
-## ComputeServerConfig
-
-Static configuration for the compute plugin
-
-Name | Description
----- | -----------
-`enableUpgradeTasks` | `boolean`<br>Whether to enable the upgrade task source, always upgrading agents to the latest version
-`withAws` | `boolean`<br>Whether to enable Amazon Web Services (AWS) specific features
-`awsRegions` | `string[]`<br>List of AWS regions for Horde to be aware of (e.g. us-east-1 or eu-central-1) Right now, this is only used for replicating CloudWatch metrics to multiple regions
-`awsAutoScalingQueueUrls` | `string[]`<br>AWS SQS queue URLs where lifecycle events from EC2 auto-scaling are received
-`fleetManagerV2` | [FleetManagerType](#fleetmanagertype-enum)<br>Default fleet manager to use (when not specified by pool)
-`fleetManagerV2Config` | `object`<br>Config for the fleet manager (serialized JSON)
-`autoEnrollAgents` | `boolean`<br>Whether to automatically enroll agents in the farm
-`defaultAgentPoolSizeStrategy` | [PoolSizeStrategy](#poolsizestrategy-enum)<br>Default agent pool sizing strategy for pools that doesn't have one explicitly configured
-`agentPoolScaleOutCooldownSeconds` | `integer`<br>Scale-out cooldown for auto-scaling agent pools (in seconds). Can be overridden by per-pool settings.
-`agentPoolScaleInCooldownSeconds` | `integer`<br>Scale-in cooldown for auto-scaling agent pools (in seconds). Can be overridden by per-pool settings.
-`computeTunnelPort` | `integer`<br>Port to listen on for tunneling compute sockets to agents
-`computeTunnelAddress` | `string`<br>What address (host:port) clients should connect to for compute socket tunneling Port may differ from  if Horde server is behind a reverse proxy/firewall
-`enabled` | `boolean`<br>Whether the plugin should be enabled or not
-
-## FleetManagerType (Enum)
-
-Available fleet managers
-
-Name | Description
----- | -----------
-`Default` | Default fleet manager
-`NoOp` | No-op fleet manager.
-`Aws` | Fleet manager for handling AWS EC2 instances. Will create and/or terminate instances from scratch.
-`AwsReuse` | Fleet manager for handling AWS EC2 instances. Will start already existing but stopped instances to reuse existing EBS disks.
-`AwsRecycle` | Fleet manager for handling AWS EC2 instances. Will start already existing but stopped instances to reuse existing EBS disks.
-`AwsAsg` | Fleet manager for handling AWS EC2 instances. Uses an EC2 auto-scaling group for controlling the number of running instances.
-
-## PoolSizeStrategy (Enum)
-
-Available pool sizing strategies
-
-Name | Description
----- | -----------
-`LeaseUtilization` | Strategy based on lease utilization
-`JobQueue` | Strategy based on size of job build queue
-`NoOp` | No-op strategy used as fallback/default behavior
-`ComputeQueueAwsMetric` | A no-op strategy that reports metrics to let an external AWS auto-scaling policy scale the fleet
-`LeaseUtilizationAwsMetric` | A no-op strategy that reports metrics to let an external AWS auto-scaling policy scale the fleet
-
-## ExperimentalServerConfig
-
-Server configuration for the experimental plugin
-
-Name | Description
----- | -----------
-`enabled` | `boolean`<br>Whether the plugin should be enabled or not
-
-## SecretsServerConfig
-
-Static configuration for the secrets plugin
-
-Name | Description
----- | -----------
-`withAws` | `boolean`<br>Whether to enable Amazon Web Services (AWS) specific features
-`enabled` | `boolean`<br>Whether the plugin should be enabled or not
-
-## StorageServerConfig
-
-Static settings for the storage system
-
-Name | Description
----- | -----------
-`bundleCacheDir` | `string`<br>Directory to use for the coarse-grained backend cache. This caches full bundles downloaded from the upstream object store.
-`bundleCacheSize` | `string`<br>Maximum size of the storage cache on disk. Accepts standard binary suffixes (kb, mb, gb, tb, etc...)
-`bundleCacheSizeBytes` | `integer`<br>Accessor for the bundle cache size in bytes
-`backends` | [BackendConfig](#backendconfig)`[]`<br>Overridden settings for storage backends. Useful for running against a production server with custom backends.
-`enabled` | `boolean`<br>Whether the plugin should be enabled or not
-
-## BackendConfig
-
-Common settings object for different providers
-
-Name | Description
----- | -----------
-`id` | `string`<br>The storage backend ID
-`base` | `string`<br>Base backend to copy default settings from
-`secondary` | `string`<br>Specifies another backend to read from if an object is not found in this one. Can be used when migrating data from one backend to another.
-`type` | [StorageBackendType](#storagebackendtype-enum)<br>
-`baseDir` | `string`<br>
-`awsBucketName` | `string`<br>Name of the bucket to use
-`awsBucketPath` | `string`<br>Base path within the bucket
-`awsCredentials` | [AwsCredentialsType](#awscredentialstype-enum)<br>Type of credentials to use
-`awsRole` | `string`<br>ARN of a role to assume
-`awsProfile` | `string`<br>The AWS profile to read credentials form
-`awsRegion` | `string`<br>Region to connect to
-`azureConnectionString` | `string`<br>Connection string for Azure
-`azureContainerName` | `string`<br>Name of the container
-`relayServer` | `string`<br>
-`relayToken` | `string`<br>
-`gcsBucketName` | `string`<br>Name of the GCS bucket to use
-`gcsBucketPath` | `string`<br>Base path within the bucket
-
-## StorageBackendType (Enum)
-
-Types of storage backend to use
-
-Name | Description
----- | -----------
-`FileSystem` | Local filesystem
-`Aws` | AWS S3
-`Azure` | Azure blob store
-`Gcs` | Google Cloud Storage
-`Memory` | In-memory only (for testing)
-
-## AwsCredentialsType (Enum)
-
-Credentials to use for AWS
-
-Name | Description
----- | -----------
-`Default` | Use default credentials from the AWS SDK
-`Profile` | Read credentials from the  profile in the AWS config file
-`AssumeRole` | Assume a particular role. Should specify ARN in
-`AssumeRoleWebIdentity` | Assume a particular role using the current environment variables.
-
 ## BuildServerConfig
 
 Static configuration for the build plugin
@@ -439,14 +281,88 @@
 `maxChildCount` | `integer`<br>Maximum number of children in each node
 `sliceThreshold` | `integer`<br>Threshold hash value for splitting interior nodes
 
-## PluginServerConfig
+## ComputeServerConfig
 
-Base class for plugin server config objects
+Static configuration for the compute plugin
 
 Name | Description
 ---- | -----------
+`enableUpgradeTasks` | `boolean`<br>Whether to enable the upgrade task source, always upgrading agents to the latest version
+`withAws` | `boolean`<br>Whether to enable Amazon Web Services (AWS) specific features
+`awsRegions` | `string[]`<br>List of AWS regions for Horde to be aware of (e.g. us-east-1 or eu-central-1) Right now, this is only used for replicating CloudWatch metrics to multiple regions
+`awsAutoScalingQueueUrls` | `string[]`<br>AWS SQS queue URLs where lifecycle events from EC2 auto-scaling are received
+`fleetManagerV2` | [FleetManagerType](#fleetmanagertype-enum)<br>Default fleet manager to use (when not specified by pool)
+`fleetManagerV2Config` | `object`<br>Config for the fleet manager (serialized JSON)
+`autoEnrollAgents` | `boolean`<br>Whether to automatically enroll agents in the farm
+`defaultAgentPoolSizeStrategy` | [PoolSizeStrategy](#poolsizestrategy-enum)<br>Default agent pool sizing strategy for pools that doesn't have one explicitly configured
+`agentPoolScaleOutCooldownSeconds` | `integer`<br>Scale-out cooldown for auto-scaling agent pools (in seconds). Can be overridden by per-pool settings.
+`agentPoolScaleInCooldownSeconds` | `integer`<br>Scale-in cooldown for auto-scaling agent pools (in seconds). Can be overridden by per-pool settings.
+`computeTunnelPort` | `integer`<br>Port to listen on for tunneling compute sockets to agents
+`computeTunnelAddress` | `string`<br>What address (host:port) clients should connect to for compute socket tunneling Port may differ from  if Horde server is behind a reverse proxy/firewall
 `enabled` | `boolean`<br>Whether the plugin should be enabled or not
 
+## FleetManagerType (Enum)
+
+Available fleet managers
+
+Name | Description
+---- | -----------
+`Default` | Default fleet manager
+`NoOp` | No-op fleet manager.
+`Aws` | Fleet manager for handling AWS EC2 instances. Will create and/or terminate instances from scratch.
+`AwsReuse` | Fleet manager for handling AWS EC2 instances. Will start already existing but stopped instances to reuse existing EBS disks.
+`AwsRecycle` | Fleet manager for handling AWS EC2 instances. Will start already existing but stopped instances to reuse existing EBS disks.
+`AwsAsg` | Fleet manager for handling AWS EC2 instances. Uses an EC2 auto-scaling group for controlling the number of running instances.
+
+## PoolSizeStrategy (Enum)
+
+Available pool sizing strategies
+
+Name | Description
+---- | -----------
+`LeaseUtilization` | Strategy based on lease utilization
+`JobQueue` | Strategy based on size of job build queue
+`NoOp` | No-op strategy used as fallback/default behavior
+`ComputeQueueAwsMetric` | A no-op strategy that reports metrics to let an external AWS auto-scaling policy scale the fleet
+`LeaseUtilizationAwsMetric` | A no-op strategy that reports metrics to let an external AWS auto-scaling policy scale the fleet
+
+## AnalyticsServerConfig
+
+Server configuration for the analytics system
+
+Name | Description
+---- | -----------
+`sinks` | [TelemetrySinkConfig](#telemetrysinkconfig)<br>Settings for the various telemetry sinks
+`enabled` | `boolean`<br>Whether the plugin should be enabled or not
+
+## TelemetrySinkConfig
+
+Telemetry sinks
+
+Name | Description
+---- | -----------
+`epic` | [EpicTelemetryConfig](#epictelemetryconfig)<br>Settings for the Epic telemetry sink
+`mongo` | [MongoTelemetryConfig](#mongotelemetryconfig)<br>Settings for the MongoDB telemetry sink
+
+## EpicTelemetryConfig
+
+Configuration for the telemetry sink
+
+Name | Description
+---- | -----------
+`url` | `string`<br>Base URL for the telemetry server
+`appId` | `string`<br>Application name to send in the event messages
+`enabled` | `boolean`<br>Whether to enable this sink
+
+## MongoTelemetryConfig
+
+Configuration for the telemetry sink
+
+Name | Description
+---- | -----------
+`retainDays` | `number`<br>Number of days worth of telmetry events to keep
+`enabled` | `boolean`<br>Whether to enable this sink
+
 ## ExamplePluginServerConfig
 
 Server configuration for the ExamplePlugin
@@ -455,6 +371,90 @@
 ---- | -----------
 `enabled` | `boolean`<br>Whether the plugin should be enabled or not
 
+## SecretsServerConfig
+
+Static configuration for the secrets plugin
+
+Name | Description
+---- | -----------
+`withAws` | `boolean`<br>Whether to enable Amazon Web Services (AWS) specific features
+`enabled` | `boolean`<br>Whether the plugin should be enabled or not
+
+## ExperimentalServerConfig
+
+Server configuration for the experimental plugin
+
+Name | Description
+---- | -----------
+`enabled` | `boolean`<br>Whether the plugin should be enabled or not
+
+## StorageServerConfig
+
+Static settings for the storage system
+
+Name | Description
+---- | -----------
+`bundleCacheDir` | `string`<br>Directory to use for the coarse-grained backend cache. This caches full bundles downloaded from the upstream object store.
+`bundleCacheSize` | `string`<br>Maximum size of the storage cache on disk. Accepts standard binary suffixes (kb, mb, gb, tb, etc...)
+`bundleCacheSizeBytes` | `integer`<br>Accessor for the bundle cache size in bytes
+`backends` | [BackendConfig](#backendconfig)`[]`<br>Overridden settings for storage backends. Useful for running against a production server with custom backends.
+`enabled` | `boolean`<br>Whether the plugin should be enabled or not
+
+## BackendConfig
+
+Common settings object for different providers
+
+Name | Description
+---- | -----------
+`id` | `string`<br>The storage backend ID
+`base` | `string`<br>Base backend to copy default settings from
+`secondary` | `string`<br>Specifies another backend to read from if an object is not found in this one. Can be used when migrating data from one backend to another.
+`type` | [StorageBackendType](#storagebackendtype-enum)<br>
+`baseDir` | `string`<br>
+`awsBucketName` | `string`<br>Name of the bucket to use
+`awsBucketPath` | `string`<br>Base path within the bucket
+`awsCredentials` | [AwsCredentialsType](#awscredentialstype-enum)<br>Type of credentials to use
+`awsRole` | `string`<br>ARN of a role to assume
+`awsProfile` | `string`<br>The AWS profile to read credentials form
+`awsRegion` | `string`<br>Region to connect to
+`azureConnectionString` | `string`<br>Connection string for Azure
+`azureContainerName` | `string`<br>Name of the container
+`relayServer` | `string`<br>
+`relayToken` | `string`<br>
+`gcsBucketName` | `string`<br>Name of the GCS bucket to use
+`gcsBucketPath` | `string`<br>Base path within the bucket
+
+## StorageBackendType (Enum)
+
+Types of storage backend to use
+
+Name | Description
+---- | -----------
+`FileSystem` | Local filesystem
+`Aws` | AWS S3
+`Azure` | Azure blob store
+`Gcs` | Google Cloud Storage
+`Memory` | In-memory only (for testing)
+
+## AwsCredentialsType (Enum)
+
+Credentials to use for AWS
+
+Name | Description
+---- | -----------
+`Default` | Use default credentials from the AWS SDK
+`Profile` | Read credentials from the  profile in the AWS config file
+`AssumeRole` | Assume a particular role. Should specify ARN in
+`AssumeRoleWebIdentity` | Assume a particular role using the current environment variables.
+
+## PluginServerConfig
+
+Base class for plugin server config objects
+
+Name | Description
+---- | -----------
+`enabled` | `boolean`<br>Whether the plugin should be enabled or not
+
 ## ToolsServerConfig
 
 Server configuration for bundled tools
