Horde: Handle client cancellation gracefully in DownloadSoftware gRPC method

Cancellations can come from client disconnects (agent restarts, network issues) and are normal.

diff --git a/Engine/Source/Programs/Horde/Plugins/Compute/HordeServer.Compute/Server/RpcService.cs b/Engine/Source/Programs/Horde/Plugins/Compute/HordeServer.Compute/Server/RpcService.cs
--- a/Engine/Source/Programs/Horde/Plugins/Compute/HordeServer.Compute/Server/RpcService.cs
+++ b/Engine/Source/Programs/Horde/Plugins/Compute/HordeServer.Compute/Server/RpcService.cs
@@ -408,9 +408,10 @@ namespace HordeServer.Server
 				throw new StructuredRpcException(StatusCode.NotFound, $"Missing tool version {version}");
 			}
 
-			await using Stream stream = await deployment.OpenZipStreamAsync(context.CancellationToken);
-			using (IMemoryOwner<byte> buffer = MemoryPool<byte>.Shared.Rent(128 * 1024))
+			try
 			{
+				await using Stream stream = await deployment.OpenZipStreamAsync(context.CancellationToken);
+				using IMemoryOwner<byte> buffer = MemoryPool<byte>.Shared.Rent(128 * 1024);
 				long totalWritten = 0;
 				for (; ; )
 				{
@@ -428,6 +429,10 @@ namespace HordeServer.Server
 				}
 				_logger.LogInformation("Agent software zip is {Size:n0} bytes", totalWritten);
 			}
+			catch (OperationCanceledException) when (context.CancellationToken.IsCancellationRequested)
+			{
+				_logger.LogDebug("DownloadSoftware stream for {ToolId} cancelled by client", toolId);
+			}
 		}
 
 		/// <inheritdoc/>
