Horde: Update release notes

#rnx

diff --git a/Engine/Source/Programs/Horde/Docs/ReleaseNotes.md b/Engine/Source/Programs/Horde/Docs/ReleaseNotes.md
--- a/Engine/Source/Programs/Horde/Docs/ReleaseNotes.md
+++ b/Engine/Source/Programs/Horde/Docs/ReleaseNotes.md
@@ -1,5 +1,21 @@
 # Release Notes
 
+## 2026-02-19
+
+* Handle client cancellation gracefully in DownloadSoftware gRPC method (51012208)
+* Fix PerforceServiceCache not picking up stream view changes (50991273)
+* Add redis-backed queue infrastructure for auto-RCA (50988984)
+* Introducing standard const for invalid commit id (50957351)
+* Shadow pool deletion from dashboard, add configVersion to dashboard config, fix IDE0052 warning (50944559)
+* Move EnableExpandedIgnorePatterns from static BuildServerConfig to dynamic BuildConfig global config (50939977)
+* Update TestExecutor to support loading ignore patterns (50933712)
+* Add temporary Pool Migration dashboard page (50933156)
+* Add TestExecutor case for matching ignore patterns (50932488)
+* Emit telemetry on lease completion (50903649)
+* Add OptionalRequires support for BuildGraph nodes (50865502)
+* Adding computed stream property to use for synthesized value (branch & stream id) (50829790)
+* V3 IncrementalPoolMigration - safe hybrid mode for incremental pool migration (50826729)
+
 ## 2026-02-12
 
 * Make the perforce change exception raise error logging (50799777)
