Horde: Clean up common properties from csproj and move them to Directory.Build.props

#rnx

#Synced-CL 50981627

#Synced-CL 50981627

diff --git a/Engine/Source/Programs/Horde/BuildHorde.xml b/Engine/Source/Programs/Horde/BuildHorde.xml
--- a/Engine/Source/Programs/Horde/BuildHorde.xml
+++ b/Engine/Source/Programs/Horde/BuildHorde.xml
@@ -93,6 +93,7 @@
 		<Node Name="Build HordeDashboard">
 			<Copy From="Engine/Source/Programs/Horde/Docs/..." To="Engine/Source/Programs/Horde/HordeDashboard/documentation/Docs/..."/>
 			<Copy From="Engine/Source/Programs/Horde/README.md" To="Engine/Source/Programs/Horde/HordeDashboard/documentation/README.md"/>
+			<Tag Files="Engine/Source/Programs/Horde/Directory.Build.props" With="#InputFiles"/>
 			<Tag Files="Engine/Source/Programs/Horde/HordeDashboard/..." Except="Engine/Source/Programs/Horde/HordeDashboard/node_modules/..." With="#InputFiles"/>
 
 			<!-- Build the dashboard Docker image and publish it -->
@@ -246,6 +247,7 @@
 			<WriteTextFile File="$(StagingDir)/Dashboard/index.html" Text="placeholder"/>			
 		
 			<!-- Build the dashboard -->
+			<Tag Files="Engine/Source/Programs/Horde/Directory.Build.props" With="#DashboardInputFiles"/>
 			<Tag Files="Engine/Source/Programs/Horde/HordeDashboard/..." With="#DashboardInputFiles"/>
 			<Docker-Build BaseDir="Engine/Source" Files="#DashboardInputFiles" Tag="hordedashboard:installer" DockerFile="Engine/Source/Programs/Horde/HordeDashboard/installer/Dockerfile" Arguments="--build-arg &quot;VersionInfo=$(InformationalVersion)&quot; --build-arg &quot;DashboardConfig=$(HordeDashboardConfig)&quot;"/>
 			
diff --git a/Engine/Source/Programs/Horde/Drivers/JobDriver.Tests/JobDriver.Tests.csproj b/Engine/Source/Programs/Horde/Drivers/JobDriver.Tests/JobDriver.Tests.csproj
--- a/Engine/Source/Programs/Horde/Drivers/JobDriver.Tests/JobDriver.Tests.csproj
+++ b/Engine/Source/Programs/Horde/Drivers/JobDriver.Tests/JobDriver.Tests.csproj
@@ -1,26 +1,8 @@
 <Project Sdk="Microsoft.NET.Sdk">
-  <Import Project="..\..\..\Shared\UnrealEngine.csproj.props" />
-
-  <PropertyGroup>
-    <TargetFramework>net8.0</TargetFramework>
-    <Nullable>enable</Nullable>
-    <IsPackable>false</IsPackable>
-    <Configurations>Debug;Release;Development;Analyze</Configurations>
-    <NoWarn>$(NoWarn);1591</NoWarn>
-  </PropertyGroup>
-
-  <PropertyGroup Condition=" '$(Configuration)|$(Platform)' == 'Analyze|AnyCPU' ">
-    <RunAnalyzersDuringBuild>True</RunAnalyzersDuringBuild>
-    <TreatWarningsAsErrors>false</TreatWarningsAsErrors>
-    <EnforceCodeStyleInBuild>True</EnforceCodeStyleInBuild>
-  </PropertyGroup>
 
   <ItemGroup>
     <PackageReference Include="Microsoft.Extensions.Logging.Abstractions" Version="8.0.3" />
-    <PackageReference Include="Microsoft.NET.Test.Sdk" Version="17.14.1" />
     <PackageReference Include="Moq" Version="4.20.72" />
-    <PackageReference Include="MSTest.TestAdapter" Version="3.11.1" />
-    <PackageReference Include="MSTest.TestFramework" Version="3.11.1" />
   </ItemGroup>
 
   <ItemGroup>
diff --git a/Engine/Source/Programs/Horde/Drivers/JobDriver/JobDriver.csproj b/Engine/Source/Programs/Horde/Drivers/JobDriver/JobDriver.csproj
--- a/Engine/Source/Programs/Horde/Drivers/JobDriver/JobDriver.csproj
+++ b/Engine/Source/Programs/Horde/Drivers/JobDriver/JobDriver.csproj
@@ -1,29 +1,16 @@
 <Project Sdk="Microsoft.NET.Sdk">
-  <Import Project="..\..\..\Shared\UnrealEngine.csproj.props" />
 
   <PropertyGroup>
     <OutputType>Exe</OutputType>
-    <TargetFramework>net8.0</TargetFramework>
-    <Nullable>enable</Nullable>
     <ImplicitUsings>enable</ImplicitUsings>
     <AssemblyName>JobDriver</AssemblyName>
     <PublishRepositoryUrl>true</PublishRepositoryUrl>
     <EmbedUntrackedSources>false</EmbedUntrackedSources>
-    <Configurations>Debug;Release;Analyze</Configurations>
     <NoWarn>$(NoWarn);NU3003</NoWarn>
   </PropertyGroup>
 
   <PropertyGroup Condition="'$(Configuration)|$(Platform)'=='Debug|AnyCPU'">
     <DocumentationFile>HordeAgentDriver.xml</DocumentationFile>
-    <TreatWarningsAsErrors>False</TreatWarningsAsErrors>
-  </PropertyGroup>
-
-  <PropertyGroup Condition=" '$(Configuration)|$(Platform)' == 'Analyze|AnyCPU' ">
-    <RunAnalyzersDuringBuild>True</RunAnalyzersDuringBuild>
-    <TreatWarningsAsErrors>false</TreatWarningsAsErrors>
-    <EnforceCodeStyleInBuild>True</EnforceCodeStyleInBuild>
-    <GenerateDocumentationFile>True</GenerateDocumentationFile>
-    <DocumentationFile></DocumentationFile>
   </PropertyGroup>
 
   <ItemGroup>
diff --git a/Engine/Source/Programs/Horde/HordeAgent.Tests/HordeAgent.Tests.csproj b/Engine/Source/Programs/Horde/HordeAgent.Tests/HordeAgent.Tests.csproj
--- a/Engine/Source/Programs/Horde/HordeAgent.Tests/HordeAgent.Tests.csproj
+++ b/Engine/Source/Programs/Horde/HordeAgent.Tests/HordeAgent.Tests.csproj
@@ -1,33 +1,8 @@
 <Project Sdk="Microsoft.NET.Sdk">
-  <Import Project="..\..\Shared\UnrealEngine.csproj.props" />
-
-  <PropertyGroup>
-    <TargetFramework>net8.0</TargetFramework>
-    <Nullable>enable</Nullable>
-    <IsPackable>false</IsPackable>
-    <Configurations>Debug;Release;Development;Analyze</Configurations>
-    <NoWarn>$(NoWarn);1591</NoWarn>
-    <RunAnalyzersDuringBuild>False</RunAnalyzersDuringBuild>
-  </PropertyGroup>
-
-  <PropertyGroup Condition=" '$(Configuration)|$(Platform)' == 'Analyze|AnyCPU' ">
-    <RunAnalyzersDuringBuild>True</RunAnalyzersDuringBuild>
-    <TreatWarningsAsErrors>false</TreatWarningsAsErrors>
-    <EnforceCodeStyleInBuild>True</EnforceCodeStyleInBuild>
-    <GenerateDocumentationFile>True</GenerateDocumentationFile>
-    <DocumentationFile></DocumentationFile>
-  </PropertyGroup>
 
   <ItemGroup>
     <PackageReference Include="Microsoft.Extensions.Logging.Abstractions" Version="8.0.3" />
-    <PackageReference Include="Microsoft.NET.Test.Sdk" Version="17.14.1" />
     <PackageReference Include="Moq" Version="4.20.72" />
-    <PackageReference Include="MSTest.TestAdapter" Version="3.11.1" />
-    <PackageReference Include="MSTest.TestFramework" Version="3.11.1" />
-    <PackageReference Include="coverlet.collector" Version="6.0.4">
-      <PrivateAssets>all</PrivateAssets>
-      <IncludeAssets>runtime; build; native; contentfiles; analyzers; buildtransitive</IncludeAssets>
-    </PackageReference>
   </ItemGroup>
 
   <ItemGroup>
diff --git a/Engine/Source/Programs/Horde/HordeAgent/HordeAgent.csproj b/Engine/Source/Programs/Horde/HordeAgent/HordeAgent.csproj
--- a/Engine/Source/Programs/Horde/HordeAgent/HordeAgent.csproj
+++ b/Engine/Source/Programs/Horde/HordeAgent/HordeAgent.csproj
@@ -1,33 +1,19 @@
 ﻿<Project Sdk="Microsoft.NET.Sdk">
-  <Import Project="..\..\Shared\UnrealEngine.csproj.props" />
 
   <PropertyGroup>
     <OutputType>Exe</OutputType>
-    <TargetFramework>net8.0</TargetFramework>
-	<RollForward>LatestMinor</RollForward>
-    <Nullable>enable</Nullable>
+    <RollForward>LatestMinor</RollForward>
     <ImplicitUsings>enable</ImplicitUsings>
-	<ApplicationIcon>HordeAgent.ico</ApplicationIcon>
+    <ApplicationIcon>HordeAgent.ico</ApplicationIcon>
     <ApplicationManifest>app.manifest</ApplicationManifest>
     <AssemblyName>HordeAgent</AssemblyName>
     <PublishRepositoryUrl>true</PublishRepositoryUrl>
     <EmbedUntrackedSources>false</EmbedUntrackedSources>
-    <Configurations>Debug;Release;Analyze</Configurations>
     <NoWarn>$(NoWarn);NU3003</NoWarn>
-    <RunAnalyzersDuringBuild>False</RunAnalyzersDuringBuild>
   </PropertyGroup>
 
   <PropertyGroup Condition="'$(Configuration)|$(Platform)'=='Debug|AnyCPU'">
     <DocumentationFile>HordeAgent.xml</DocumentationFile>
-    <TreatWarningsAsErrors>False</TreatWarningsAsErrors>
-  </PropertyGroup>
-
-  <PropertyGroup Condition=" '$(Configuration)|$(Platform)' == 'Analyze|AnyCPU' ">
-    <RunAnalyzersDuringBuild>True</RunAnalyzersDuringBuild>
-    <TreatWarningsAsErrors>false</TreatWarningsAsErrors>
-    <EnforceCodeStyleInBuild>True</EnforceCodeStyleInBuild>
-    <GenerateDocumentationFile>True</GenerateDocumentationFile>
-    <DocumentationFile></DocumentationFile>
   </PropertyGroup>
 
   <ItemGroup>
diff --git a/Engine/Source/Programs/Horde/HordeCmd/HordeCmd.csproj b/Engine/Source/Programs/Horde/HordeCmd/HordeCmd.csproj
--- a/Engine/Source/Programs/Horde/HordeCmd/HordeCmd.csproj
+++ b/Engine/Source/Programs/Horde/HordeCmd/HordeCmd.csproj
@@ -1,22 +1,9 @@
 <Project Sdk="Microsoft.NET.Sdk">
-  <Import Project="..\..\Shared\UnrealEngine.csproj.props" />
 
   <PropertyGroup>
     <OutputType>Exe</OutputType>
     <AssemblyName>Horde</AssemblyName>
-    <TargetFramework>net8.0</TargetFramework>
     <ImplicitUsings>enable</ImplicitUsings>
-    <Nullable>enable</Nullable>
-    <Configurations>Debug;Release;Analyze</Configurations>
-    <RunAnalyzersDuringBuild>False</RunAnalyzersDuringBuild>
-  </PropertyGroup>
-
-  <PropertyGroup Condition=" '$(Configuration)|$(Platform)' == 'Analyze|AnyCPU' ">
-    <RunAnalyzersDuringBuild>True</RunAnalyzersDuringBuild>
-    <TreatWarningsAsErrors>False</TreatWarningsAsErrors>
-    <EnforceCodeStyleInBuild>True</EnforceCodeStyleInBuild>
-    <GenerateDocumentationFile>True</GenerateDocumentationFile>
-    <DocumentationFile></DocumentationFile>
   </PropertyGroup>
 
   <ItemGroup>
diff --git a/Engine/Source/Programs/Horde/HordeServer.Shared/HordeServer.Shared.csproj b/Engine/Source/Programs/Horde/HordeServer.Shared/HordeServer.Shared.csproj
--- a/Engine/Source/Programs/Horde/HordeServer.Shared/HordeServer.Shared.csproj
+++ b/Engine/Source/Programs/Horde/HordeServer.Shared/HordeServer.Shared.csproj
@@ -1,21 +1,8 @@
 ﻿<Project Sdk="Microsoft.NET.Sdk">
-  <Import Project="..\..\Shared\UnrealEngine.csproj.props" />
 
   <PropertyGroup>
-    <TargetFramework>net8.0</TargetFramework>
     <ImplicitUsings>enable</ImplicitUsings>
-    <Nullable>enable</Nullable>
     <RootNamespace>HordeServer</RootNamespace>
-    <RunAnalyzersDuringBuild>False</RunAnalyzersDuringBuild>
-    <Configurations>Debug;Release;Analyze</Configurations>
-  </PropertyGroup>
-
-  <PropertyGroup Condition=" '$(Configuration)|$(Platform)' == 'Analyze|AnyCPU' ">
-    <RunAnalyzersDuringBuild>True</RunAnalyzersDuringBuild>
-    <TreatWarningsAsErrors>false</TreatWarningsAsErrors>
-    <EnforceCodeStyleInBuild>True</EnforceCodeStyleInBuild>
-    <GenerateDocumentationFile>True</GenerateDocumentationFile>
-    <DocumentationFile></DocumentationFile>
   </PropertyGroup>
 
   <ItemGroup>
diff --git a/Engine/Source/Programs/Horde/HordeServer.Tests.Shared/HordeServer.Tests.Shared.csproj b/Engine/Source/Programs/Horde/HordeServer.Tests.Shared/HordeServer.Tests.Shared.csproj
--- a/Engine/Source/Programs/Horde/HordeServer.Tests.Shared/HordeServer.Tests.Shared.csproj
+++ b/Engine/Source/Programs/Horde/HordeServer.Tests.Shared/HordeServer.Tests.Shared.csproj
@@ -1,21 +1,8 @@
 ﻿<Project Sdk="Microsoft.NET.Sdk">
-  <Import Project="..\..\Shared\UnrealEngine.csproj.props" />
 
   <PropertyGroup>
-    <TargetFramework>net8.0</TargetFramework>
     <ImplicitUsings>enable</ImplicitUsings>
-    <Nullable>enable</Nullable>
     <RootNamespace>HordeServer.Tests</RootNamespace>
-    <RunAnalyzersDuringBuild>False</RunAnalyzersDuringBuild>
-    <Configurations>Debug;Release;Analyze</Configurations>
-    <NoWarn>$(NoWarn);1591</NoWarn>
-  </PropertyGroup>
-
-  <PropertyGroup Condition=" '$(Configuration)|$(Platform)' == 'Analyze|AnyCPU' ">
-    <RunAnalyzersDuringBuild>True</RunAnalyzersDuringBuild>
-    <TreatWarningsAsErrors>false</TreatWarningsAsErrors>
-    <EnforceCodeStyleInBuild>True</EnforceCodeStyleInBuild>
-    <GenerateDocumentationFile>True</GenerateDocumentationFile>
   </PropertyGroup>
 
   <ItemGroup>
@@ -24,15 +11,8 @@
 
   <ItemGroup>
     <PackageReference Include="Microsoft.AspNetCore.Mvc.Testing" Version="8.0.24" />
-    <PackageReference Include="Microsoft.NET.Test.Sdk" Version="17.14.1" />
     <PackageReference Include="Moq" Version="4.20.72" />
-    <PackageReference Include="MSTest.TestAdapter" Version="3.11.1" />
-    <PackageReference Include="MSTest.TestFramework" Version="3.11.1" />
     <PackageReference Include="Microsoft.Extensions.Caching.Memory" Version="8.0.1" />
-    <PackageReference Include="coverlet.collector" Version="6.0.4">
-      <PrivateAssets>all</PrivateAssets>
-      <IncludeAssets>runtime; build; native; contentfiles; analyzers; buildtransitive</IncludeAssets>
-    </PackageReference>
   </ItemGroup>
 
 </Project>
diff --git a/Engine/Source/Programs/Horde/HordeServer.Tests/HordeServer.Tests.csproj b/Engine/Source/Programs/Horde/HordeServer.Tests/HordeServer.Tests.csproj
--- a/Engine/Source/Programs/Horde/HordeServer.Tests/HordeServer.Tests.csproj
+++ b/Engine/Source/Programs/Horde/HordeServer.Tests/HordeServer.Tests.csproj
@@ -1,35 +1,14 @@
 <Project Sdk="Microsoft.NET.Sdk">
-  <Import Project="..\..\Shared\UnrealEngine.csproj.props" />
 
   <PropertyGroup>
-    <TargetFramework>net8.0</TargetFramework>
-    <Nullable>enable</Nullable>
-    <IsPackable>false</IsPackable>
     <NeutralLanguage>en-US</NeutralLanguage>
-    <Configurations>Debug;Release;Development;Analyze</Configurations>
-    <NoWarn>$(NoWarn);1591</NoWarn>
     <GenerateProgramFile>false</GenerateProgramFile>
-    <RunAnalyzersDuringBuild>False</RunAnalyzersDuringBuild>
-  </PropertyGroup>
-
-  <PropertyGroup Condition=" '$(Configuration)|$(Platform)' == 'Analyze|AnyCPU' ">
-    <RunAnalyzersDuringBuild>True</RunAnalyzersDuringBuild>
-    <TreatWarningsAsErrors>false</TreatWarningsAsErrors>
-    <EnforceCodeStyleInBuild>True</EnforceCodeStyleInBuild>
-    <GenerateDocumentationFile>True</GenerateDocumentationFile>
   </PropertyGroup>
 
   <ItemGroup>
     <PackageReference Include="Microsoft.AspNetCore.Mvc.Testing" Version="8.0.24" />
-    <PackageReference Include="Microsoft.NET.Test.Sdk" Version="17.14.1" />
     <PackageReference Include="Moq" Version="4.20.72" />
-    <PackageReference Include="MSTest.TestAdapter" Version="3.11.1" />
-    <PackageReference Include="MSTest.TestFramework" Version="3.11.1" />
     <PackageReference Include="Microsoft.Extensions.Caching.Memory" Version="8.0.1" />
-    <PackageReference Include="coverlet.collector" Version="6.0.4">
-      <PrivateAssets>all</PrivateAssets>
-      <IncludeAssets>runtime; build; native; contentfiles; analyzers; buildtransitive</IncludeAssets>
-    </PackageReference>
   </ItemGroup>
 
   <ItemGroup>
diff --git a/Engine/Source/Programs/Horde/HordeServer/Dockerfile b/Engine/Source/Programs/Horde/HordeServer/Dockerfile
--- a/Engine/Source/Programs/Horde/HordeServer/Dockerfile
+++ b/Engine/Source/Programs/Horde/HordeServer/Dockerfile
@@ -31,7 +31,7 @@ RUN wget https://repo.mongodb.org/apt/debian/dists/bookworm/mongodb-org/7.0/main
 # Copy build dependencies
 COPY Binaries/DotNET/EpicGames.Perforce.Native/ Binaries/DotNET/EpicGames.Perforce.Native/
 COPY Source/Programs/Shared/ ./Source/Programs/Shared/
-COPY Source/Programs/Horde/.editorconfig Source/Programs/Horde/.editorconfig
+COPY Source/Programs/Horde/* Source/Programs/Horde/
 COPY Source/Programs/Horde/Plugins/ Source/Programs/Horde/Plugins/
 COPY Source/Programs/Horde/HordeServer/ Source/Programs/Horde/HordeServer/
 COPY Source/Programs/Horde/HordeServer.Shared/ Source/Programs/Horde/HordeServer.Shared/
diff --git a/Engine/Source/Programs/Horde/HordeServer/HordeServer.csproj b/Engine/Source/Programs/Horde/HordeServer/HordeServer.csproj
--- a/Engine/Source/Programs/Horde/HordeServer/HordeServer.csproj
+++ b/Engine/Source/Programs/Horde/HordeServer/HordeServer.csproj
@@ -1,23 +1,9 @@
 ﻿<Project Sdk="Microsoft.NET.Sdk.Web">
-  <Import Project="..\..\Shared\UnrealEngine.csproj.props" />
 
   <PropertyGroup>
-    <TargetFramework>net8.0</TargetFramework>
-    <Nullable>enable</Nullable>
-    <Deterministic>true</Deterministic>
     <AssemblyName>HordeServer</AssemblyName>
     <DocumentationFile>HordeServer.xml</DocumentationFile>
     <ApplicationIcon>HordeServer.ico</ApplicationIcon>
-    <RunAnalyzersDuringBuild>False</RunAnalyzersDuringBuild>
-    <Configurations>Debug;Release;Analyze</Configurations>
-  </PropertyGroup>
-
-  <PropertyGroup Condition=" '$(Configuration)|$(Platform)' == 'Analyze|AnyCPU' ">
-    <RunAnalyzersDuringBuild>True</RunAnalyzersDuringBuild>
-    <TreatWarningsAsErrors>false</TreatWarningsAsErrors>
-    <EnforceCodeStyleInBuild>True</EnforceCodeStyleInBuild>
-    <GenerateDocumentationFile>True</GenerateDocumentationFile>
-    <DocumentationFile></DocumentationFile>
   </PropertyGroup>
 
   <ItemGroup>
diff --git a/Engine/Source/Programs/Horde/Plugins/Analytics/HordeServer.Analytics.Integration.Tests/HordeServer.Analytics.Integration.Tests.csproj b/Engine/Source/Programs/Horde/Plugins/Analytics/HordeServer.Analytics.Integration.Tests/HordeServer.Analytics.Integration.Tests.csproj
--- a/Engine/Source/Programs/Horde/Plugins/Analytics/HordeServer.Analytics.Integration.Tests/HordeServer.Analytics.Integration.Tests.csproj
+++ b/Engine/Source/Programs/Horde/Plugins/Analytics/HordeServer.Analytics.Integration.Tests/HordeServer.Analytics.Integration.Tests.csproj
@@ -1,33 +1,10 @@
 <Project Sdk="Microsoft.NET.Sdk">
-  <Import Project="..\..\..\..\Shared\UnrealEngine.csproj.props" />
 
   <PropertyGroup>
-    <TargetFramework>net8.0</TargetFramework>
     <ImplicitUsings>enable</ImplicitUsings>
-    <Nullable>enable</Nullable>
-    <IsPackable>false</IsPackable>
-    <IsTestProject>true</IsTestProject>
-    <RunAnalyzersDuringBuild>False</RunAnalyzersDuringBuild>
-    <Configurations>Debug;Release;Analyze</Configurations>
-    <NoWarn>$(NoWarn);1591</NoWarn>
   </PropertyGroup>
 
-  <PropertyGroup Condition=" '$(Configuration)|$(Platform)' == 'Analyze|AnyCPU' ">
-    <RunAnalyzersDuringBuild>True</RunAnalyzersDuringBuild>
-    <TreatWarningsAsErrors>false</TreatWarningsAsErrors>
-    <EnforceCodeStyleInBuild>True</EnforceCodeStyleInBuild>
-    <GenerateDocumentationFile>True</GenerateDocumentationFile>
-  </PropertyGroup>
 
-  <ItemGroup>
-    <PackageReference Include="coverlet.collector" Version="6.0.4">
-      <PrivateAssets>all</PrivateAssets>
-      <IncludeAssets>runtime; build; native; contentfiles; analyzers; buildtransitive</IncludeAssets>
-    </PackageReference>
-    <PackageReference Include="Microsoft.NET.Test.Sdk" Version="17.14.1" />
-    <PackageReference Include="MSTest.TestAdapter" Version="3.11.1" />
-    <PackageReference Include="MSTest.TestFramework" Version="3.11.1" />
-  </ItemGroup>
 
   <ItemGroup>
     <ProjectReference Include="..\..\Build\HordeServer.Build.Tests\HordeServer.Build.Tests.csproj" />
diff --git a/Engine/Source/Programs/Horde/Plugins/Analytics/HordeServer.Analytics.Tests/HordeServer.Analytics.Tests.csproj b/Engine/Source/Programs/Horde/Plugins/Analytics/HordeServer.Analytics.Tests/HordeServer.Analytics.Tests.csproj
--- a/Engine/Source/Programs/Horde/Plugins/Analytics/HordeServer.Analytics.Tests/HordeServer.Analytics.Tests.csproj
+++ b/Engine/Source/Programs/Horde/Plugins/Analytics/HordeServer.Analytics.Tests/HordeServer.Analytics.Tests.csproj
@@ -1,34 +1,9 @@
 <Project Sdk="Microsoft.NET.Sdk">
-  <Import Project="..\..\..\..\Shared\UnrealEngine.csproj.props" />
 
   <PropertyGroup>
-    <TargetFramework>net8.0</TargetFramework>
     <ImplicitUsings>enable</ImplicitUsings>
-    <Nullable>enable</Nullable>
-    <IsPackable>false</IsPackable>
-    <IsTestProject>true</IsTestProject>
-    <RunAnalyzersDuringBuild>False</RunAnalyzersDuringBuild>
-    <Configurations>Debug;Release;Analyze</Configurations>
-    <NoWarn>$(NoWarn);1591</NoWarn>
   </PropertyGroup>
 
-  <PropertyGroup Condition=" '$(Configuration)|$(Platform)' == 'Analyze|AnyCPU' ">
-    <RunAnalyzersDuringBuild>True</RunAnalyzersDuringBuild>
-    <TreatWarningsAsErrors>false</TreatWarningsAsErrors>
-    <EnforceCodeStyleInBuild>True</EnforceCodeStyleInBuild>
-    <GenerateDocumentationFile>True</GenerateDocumentationFile>
-  </PropertyGroup>
-
-  <ItemGroup>
-    <PackageReference Include="coverlet.collector" Version="6.0.4">
-      <PrivateAssets>all</PrivateAssets>
-      <IncludeAssets>runtime; build; native; contentfiles; analyzers; buildtransitive</IncludeAssets>
-    </PackageReference>
-    <PackageReference Include="Microsoft.NET.Test.Sdk" Version="17.14.1" />
-    <PackageReference Include="MSTest.TestAdapter" Version="3.11.1" />
-    <PackageReference Include="MSTest.TestFramework" Version="3.11.1" />
-  </ItemGroup>
-
   <ItemGroup>
     <ProjectReference Include="..\HordeServer.Analytics\HordeServer.Analytics.csproj" />
     <ProjectReference Include="..\..\..\HordeServer.Tests.Shared\HordeServer.Tests.Shared.csproj" />
diff --git a/Engine/Source/Programs/Horde/Plugins/Analytics/HordeServer.Analytics/HordeServer.Analytics.csproj b/Engine/Source/Programs/Horde/Plugins/Analytics/HordeServer.Analytics/HordeServer.Analytics.csproj
--- a/Engine/Source/Programs/Horde/Plugins/Analytics/HordeServer.Analytics/HordeServer.Analytics.csproj
+++ b/Engine/Source/Programs/Horde/Plugins/Analytics/HordeServer.Analytics/HordeServer.Analytics.csproj
@@ -1,21 +1,8 @@
 ﻿<Project Sdk="Microsoft.NET.Sdk">
-  <Import Project="..\..\..\..\Shared\UnrealEngine.csproj.props" />
 
   <PropertyGroup>
-    <TargetFramework>net8.0</TargetFramework>
     <ImplicitUsings>enable</ImplicitUsings>
-    <Nullable>enable</Nullable>
     <RootNamespace>HordeServer</RootNamespace>
-    <RunAnalyzersDuringBuild>False</RunAnalyzersDuringBuild>
-    <Configurations>Debug;Release;Analyze</Configurations>
-  </PropertyGroup>
-
-  <PropertyGroup Condition=" '$(Configuration)|$(Platform)' == 'Analyze|AnyCPU' ">
-    <RunAnalyzersDuringBuild>True</RunAnalyzersDuringBuild>
-    <TreatWarningsAsErrors>false</TreatWarningsAsErrors>
-    <EnforceCodeStyleInBuild>True</EnforceCodeStyleInBuild>
-    <GenerateDocumentationFile>True</GenerateDocumentationFile>
-    <DocumentationFile></DocumentationFile>
   </PropertyGroup>
 
   <ItemGroup>
diff --git a/Engine/Source/Programs/Horde/Plugins/Build/HordeServer.Build.Tests/HordeServer.Build.Tests.csproj b/Engine/Source/Programs/Horde/Plugins/Build/HordeServer.Build.Tests/HordeServer.Build.Tests.csproj
--- a/Engine/Source/Programs/Horde/Plugins/Build/HordeServer.Build.Tests/HordeServer.Build.Tests.csproj
+++ b/Engine/Source/Programs/Horde/Plugins/Build/HordeServer.Build.Tests/HordeServer.Build.Tests.csproj
@@ -1,33 +1,11 @@
 <Project Sdk="Microsoft.NET.Sdk">
-  <Import Project="..\..\..\..\Shared\UnrealEngine.csproj.props" />
 
   <PropertyGroup>
-    <TargetFramework>net8.0</TargetFramework>
     <ImplicitUsings>enable</ImplicitUsings>
-    <Nullable>enable</Nullable>
-    <IsPackable>false</IsPackable>
-    <IsTestProject>true</IsTestProject>
     <RootNamespace>HordeServer.Tests</RootNamespace>
-    <RunAnalyzersDuringBuild>False</RunAnalyzersDuringBuild>
-    <Configurations>Debug;Release;Analyze</Configurations>
-    <NoWarn>$(NoWarn);1591</NoWarn>
-  </PropertyGroup>
-
-  <PropertyGroup Condition=" '$(Configuration)|$(Platform)' == 'Analyze|AnyCPU' ">
-    <RunAnalyzersDuringBuild>True</RunAnalyzersDuringBuild>
-    <TreatWarningsAsErrors>false</TreatWarningsAsErrors>
-    <EnforceCodeStyleInBuild>True</EnforceCodeStyleInBuild>
-    <GenerateDocumentationFile>True</GenerateDocumentationFile>
   </PropertyGroup>
 
   <ItemGroup>
-    <PackageReference Include="coverlet.collector" Version="6.0.4">
-      <PrivateAssets>all</PrivateAssets>
-      <IncludeAssets>runtime; build; native; contentfiles; analyzers; buildtransitive</IncludeAssets>
-    </PackageReference>
-    <PackageReference Include="Microsoft.NET.Test.Sdk" Version="17.14.1" />
-    <PackageReference Include="MSTest.TestAdapter" Version="3.11.1" />
-    <PackageReference Include="MSTest.TestFramework" Version="3.11.1" />
     <PackageReference Include="Microsoft.AspNetCore.Mvc.Testing" Version="8.0.24" />
     <PackageReference Include="Moq" Version="4.20.72" />
     <PackageReference Include="Microsoft.Extensions.Caching.Memory" Version="8.0.1" />
diff --git a/Engine/Source/Programs/Horde/Plugins/Build/HordeServer.Build/HordeServer.Build.csproj b/Engine/Source/Programs/Horde/Plugins/Build/HordeServer.Build/HordeServer.Build.csproj
--- a/Engine/Source/Programs/Horde/Plugins/Build/HordeServer.Build/HordeServer.Build.csproj
+++ b/Engine/Source/Programs/Horde/Plugins/Build/HordeServer.Build/HordeServer.Build.csproj
@@ -1,21 +1,8 @@
 ﻿<Project Sdk="Microsoft.NET.Sdk">
-  <Import Project="..\..\..\..\Shared\UnrealEngine.csproj.props" />
 
   <PropertyGroup>
-    <TargetFramework>net8.0</TargetFramework>
     <ImplicitUsings>enable</ImplicitUsings>
-    <Nullable>enable</Nullable>
     <RootNamespace>HordeServer</RootNamespace>
-    <RunAnalyzersDuringBuild>False</RunAnalyzersDuringBuild>
-    <Configurations>Debug;Release;Analyze</Configurations>
-  </PropertyGroup>
-
-  <PropertyGroup Condition=" '$(Configuration)|$(Platform)' == 'Analyze|AnyCPU' ">
-    <RunAnalyzersDuringBuild>True</RunAnalyzersDuringBuild>
-    <TreatWarningsAsErrors>false</TreatWarningsAsErrors>
-    <EnforceCodeStyleInBuild>True</EnforceCodeStyleInBuild>
-    <GenerateDocumentationFile>True</GenerateDocumentationFile>
-    <DocumentationFile></DocumentationFile>
   </PropertyGroup>
 
   <ItemGroup>
diff --git a/Engine/Source/Programs/Horde/Plugins/Compute/HordeServer.Compute.Tests/HordeServer.Compute.Tests.csproj b/Engine/Source/Programs/Horde/Plugins/Compute/HordeServer.Compute.Tests/HordeServer.Compute.Tests.csproj
--- a/Engine/Source/Programs/Horde/Plugins/Compute/HordeServer.Compute.Tests/HordeServer.Compute.Tests.csproj
+++ b/Engine/Source/Programs/Horde/Plugins/Compute/HordeServer.Compute.Tests/HordeServer.Compute.Tests.csproj
@@ -1,35 +1,10 @@
 <Project Sdk="Microsoft.NET.Sdk">
-  <Import Project="..\..\..\..\Shared\UnrealEngine.csproj.props" />
 
   <PropertyGroup>
-    <TargetFramework>net8.0</TargetFramework>
     <ImplicitUsings>enable</ImplicitUsings>
-    <Nullable>enable</Nullable>
-    <IsPackable>false</IsPackable>
-    <IsTestProject>true</IsTestProject>
     <RootNamespace>HordeServer.Tests</RootNamespace>
-    <RunAnalyzersDuringBuild>False</RunAnalyzersDuringBuild>
-    <Configurations>Debug;Release;Analyze</Configurations>
-    <NoWarn>$(NoWarn);1591</NoWarn>
   </PropertyGroup>
 
-  <PropertyGroup Condition=" '$(Configuration)|$(Platform)' == 'Analyze|AnyCPU' ">
-    <RunAnalyzersDuringBuild>True</RunAnalyzersDuringBuild>
-    <TreatWarningsAsErrors>false</TreatWarningsAsErrors>
-    <EnforceCodeStyleInBuild>True</EnforceCodeStyleInBuild>
-    <GenerateDocumentationFile>True</GenerateDocumentationFile>
-  </PropertyGroup>
-
-  <ItemGroup>
-    <PackageReference Include="coverlet.collector" Version="6.0.4">
-      <PrivateAssets>all</PrivateAssets>
-      <IncludeAssets>runtime; build; native; contentfiles; analyzers; buildtransitive</IncludeAssets>
-    </PackageReference>
-    <PackageReference Include="Microsoft.NET.Test.Sdk" Version="17.14.1" />
-    <PackageReference Include="MSTest.TestAdapter" Version="3.11.1" />
-    <PackageReference Include="MSTest.TestFramework" Version="3.11.1" />
-  </ItemGroup>
-
   <ItemGroup>
     <ProjectReference Include="..\HordeServer.Compute\HordeServer.Compute.csproj" />
     <ProjectReference Include="..\..\Storage\HordeServer.Storage\HordeServer.Storage.csproj" />
diff --git a/Engine/Source/Programs/Horde/Plugins/Compute/HordeServer.Compute/HordeServer.Compute.csproj b/Engine/Source/Programs/Horde/Plugins/Compute/HordeServer.Compute/HordeServer.Compute.csproj
--- a/Engine/Source/Programs/Horde/Plugins/Compute/HordeServer.Compute/HordeServer.Compute.csproj
+++ b/Engine/Source/Programs/Horde/Plugins/Compute/HordeServer.Compute/HordeServer.Compute.csproj
@@ -1,21 +1,8 @@
 ﻿<Project Sdk="Microsoft.NET.Sdk">
-  <Import Project="..\..\..\..\Shared\UnrealEngine.csproj.props" />
 
   <PropertyGroup>
-    <TargetFramework>net8.0</TargetFramework>
     <ImplicitUsings>enable</ImplicitUsings>
-    <Nullable>enable</Nullable>
     <RootNamespace>HordeServer</RootNamespace>
-    <RunAnalyzersDuringBuild>False</RunAnalyzersDuringBuild>
-    <Configurations>Debug;Release;Analyze</Configurations>
-  </PropertyGroup>
-
-  <PropertyGroup Condition=" '$(Configuration)|$(Platform)' == 'Analyze|AnyCPU' ">
-    <RunAnalyzersDuringBuild>True</RunAnalyzersDuringBuild>
-    <TreatWarningsAsErrors>false</TreatWarningsAsErrors>
-    <EnforceCodeStyleInBuild>True</EnforceCodeStyleInBuild>
-    <GenerateDocumentationFile>True</GenerateDocumentationFile>
-    <DocumentationFile></DocumentationFile>
   </PropertyGroup>
 
   <ItemGroup>
diff --git a/Engine/Source/Programs/Horde/Plugins/ExamplePlugin/HordeServer.ExamplePlugin.Tests/HordeServer.ExamplePlugin.Tests.csproj b/Engine/Source/Programs/Horde/Plugins/ExamplePlugin/HordeServer.ExamplePlugin.Tests/HordeServer.ExamplePlugin.Tests.csproj
--- a/Engine/Source/Programs/Horde/Plugins/ExamplePlugin/HordeServer.ExamplePlugin.Tests/HordeServer.ExamplePlugin.Tests.csproj
+++ b/Engine/Source/Programs/Horde/Plugins/ExamplePlugin/HordeServer.ExamplePlugin.Tests/HordeServer.ExamplePlugin.Tests.csproj
@@ -1,33 +1,11 @@
 ﻿<Project Sdk="Microsoft.NET.Sdk">
-  <Import Project="..\..\..\..\Shared\UnrealEngine.csproj.props" />
 
   <PropertyGroup>
-    <TargetFramework>net8.0</TargetFramework>
     <ImplicitUsings>enable</ImplicitUsings>
-    <Nullable>enable</Nullable>
-    <IsPackable>false</IsPackable>
-    <IsTestProject>true</IsTestProject>
     <RootNamespace>HordeServer.Tests</RootNamespace>
-    <RunAnalyzersDuringBuild>False</RunAnalyzersDuringBuild>
-    <Configurations>Debug;Release;Analyze</Configurations>
-    <NoWarn>$(NoWarn);1591</NoWarn>
-  </PropertyGroup>
-
-  <PropertyGroup Condition=" '$(Configuration)|$(Platform)' == 'Analyze|AnyCPU' ">
-    <RunAnalyzersDuringBuild>True</RunAnalyzersDuringBuild>
-    <TreatWarningsAsErrors>false</TreatWarningsAsErrors>
-    <EnforceCodeStyleInBuild>True</EnforceCodeStyleInBuild>
-    <GenerateDocumentationFile>True</GenerateDocumentationFile>
   </PropertyGroup>
 
   <ItemGroup>
-    <PackageReference Include="coverlet.collector" Version="6.0.4">
-      <PrivateAssets>all</PrivateAssets>
-      <IncludeAssets>runtime; build; native; contentfiles; analyzers; buildtransitive</IncludeAssets>
-    </PackageReference>
-    <PackageReference Include="Microsoft.NET.Test.Sdk" Version="17.14.1" />
-    <PackageReference Include="MSTest.TestAdapter" Version="3.11.1" />
-    <PackageReference Include="MSTest.TestFramework" Version="3.11.1" />
     <PackageReference Include="Microsoft.AspNetCore.Mvc.Testing" Version="8.0.24" />
     <PackageReference Include="Moq" Version="4.20.72" />
     <PackageReference Include="Microsoft.Extensions.Caching.Memory" Version="8.0.1" />
diff --git a/Engine/Source/Programs/Horde/Plugins/ExamplePlugin/HordeServer.ExamplePlugin/HordeServer.ExamplePlugin.csproj b/Engine/Source/Programs/Horde/Plugins/ExamplePlugin/HordeServer.ExamplePlugin/HordeServer.ExamplePlugin.csproj
--- a/Engine/Source/Programs/Horde/Plugins/ExamplePlugin/HordeServer.ExamplePlugin/HordeServer.ExamplePlugin.csproj
+++ b/Engine/Source/Programs/Horde/Plugins/ExamplePlugin/HordeServer.ExamplePlugin/HordeServer.ExamplePlugin.csproj
@@ -1,21 +1,8 @@
 ﻿<Project Sdk="Microsoft.NET.Sdk">
-  <Import Project="..\..\..\..\Shared\UnrealEngine.csproj.props" />
 
   <PropertyGroup>
-    <TargetFramework>net8.0</TargetFramework>
     <ImplicitUsings>enable</ImplicitUsings>
-    <Nullable>enable</Nullable>
     <RootNamespace>HordeServer</RootNamespace>
-    <RunAnalyzersDuringBuild>False</RunAnalyzersDuringBuild>
-    <Configurations>Debug;Release;Analyze</Configurations>
-  </PropertyGroup>
-
-  <PropertyGroup Condition=" '$(Configuration)|$(Platform)' == 'Analyze|AnyCPU' ">
-    <RunAnalyzersDuringBuild>True</RunAnalyzersDuringBuild>
-    <TreatWarningsAsErrors>false</TreatWarningsAsErrors>
-    <EnforceCodeStyleInBuild>True</EnforceCodeStyleInBuild>
-    <GenerateDocumentationFile>True</GenerateDocumentationFile>
-    <DocumentationFile></DocumentationFile>
   </PropertyGroup>
 
   <ItemGroup>
diff --git a/Engine/Source/Programs/Horde/Plugins/Experimental/HordeServer.Experimental.Tests/HordeServer.Experimental.Tests.csproj b/Engine/Source/Programs/Horde/Plugins/Experimental/HordeServer.Experimental.Tests/HordeServer.Experimental.Tests.csproj
--- a/Engine/Source/Programs/Horde/Plugins/Experimental/HordeServer.Experimental.Tests/HordeServer.Experimental.Tests.csproj
+++ b/Engine/Source/Programs/Horde/Plugins/Experimental/HordeServer.Experimental.Tests/HordeServer.Experimental.Tests.csproj
@@ -1,33 +1,11 @@
 <Project Sdk="Microsoft.NET.Sdk">
-  <Import Project="..\..\..\..\Shared\UnrealEngine.csproj.props" />
 
   <PropertyGroup>
-    <TargetFramework>net8.0</TargetFramework>
     <ImplicitUsings>enable</ImplicitUsings>
-    <Nullable>enable</Nullable>
-    <IsPackable>false</IsPackable>
-    <IsTestProject>true</IsTestProject>
     <RootNamespace>HordeServer.Tests</RootNamespace>
-    <RunAnalyzersDuringBuild>False</RunAnalyzersDuringBuild>
-    <Configurations>Debug;Release;Analyze</Configurations>
-    <NoWarn>$(NoWarn);1591</NoWarn>
-  </PropertyGroup>
-
-  <PropertyGroup Condition=" '$(Configuration)|$(Platform)' == 'Analyze|AnyCPU' ">
-    <RunAnalyzersDuringBuild>True</RunAnalyzersDuringBuild>
-    <TreatWarningsAsErrors>false</TreatWarningsAsErrors>
-    <EnforceCodeStyleInBuild>True</EnforceCodeStyleInBuild>
-    <GenerateDocumentationFile>True</GenerateDocumentationFile>
   </PropertyGroup>
 
   <ItemGroup>
-    <PackageReference Include="coverlet.collector" Version="6.0.4">
-      <PrivateAssets>all</PrivateAssets>
-      <IncludeAssets>runtime; build; native; contentfiles; analyzers; buildtransitive</IncludeAssets>
-    </PackageReference>
-    <PackageReference Include="Microsoft.NET.Test.Sdk" Version="17.14.1" />
-    <PackageReference Include="MSTest.TestAdapter" Version="3.11.1" />
-    <PackageReference Include="MSTest.TestFramework" Version="3.11.1" />
     <PackageReference Include="Microsoft.AspNetCore.Mvc.Testing" Version="8.0.24" />
     <PackageReference Include="Moq" Version="4.20.72" />
     <PackageReference Include="Microsoft.Extensions.Caching.Memory" Version="8.0.1" />
diff --git a/Engine/Source/Programs/Horde/Plugins/Experimental/HordeServer.Experimental/HordeServer.Experimental.csproj b/Engine/Source/Programs/Horde/Plugins/Experimental/HordeServer.Experimental/HordeServer.Experimental.csproj
--- a/Engine/Source/Programs/Horde/Plugins/Experimental/HordeServer.Experimental/HordeServer.Experimental.csproj
+++ b/Engine/Source/Programs/Horde/Plugins/Experimental/HordeServer.Experimental/HordeServer.Experimental.csproj
@@ -1,21 +1,8 @@
 ﻿<Project Sdk="Microsoft.NET.Sdk">
-  <Import Project="..\..\..\..\Shared\UnrealEngine.csproj.props" />
 
   <PropertyGroup>
-    <TargetFramework>net8.0</TargetFramework>
     <ImplicitUsings>enable</ImplicitUsings>
-    <Nullable>enable</Nullable>
     <RootNamespace>HordeServer</RootNamespace>
-    <RunAnalyzersDuringBuild>False</RunAnalyzersDuringBuild>
-    <Configurations>Debug;Release;Analyze</Configurations>
-  </PropertyGroup>
-
-  <PropertyGroup Condition=" '$(Configuration)|$(Platform)' == 'Analyze|AnyCPU' ">
-    <RunAnalyzersDuringBuild>True</RunAnalyzersDuringBuild>
-    <TreatWarningsAsErrors>false</TreatWarningsAsErrors>
-    <EnforceCodeStyleInBuild>True</EnforceCodeStyleInBuild>
-    <GenerateDocumentationFile>True</GenerateDocumentationFile>
-    <DocumentationFile></DocumentationFile>
   </PropertyGroup>
 
   <ItemGroup>
diff --git a/Engine/Source/Programs/Horde/Plugins/Secrets/HordeServer.Secrets.Tests/HordeServer.Secrets.Tests.csproj b/Engine/Source/Programs/Horde/Plugins/Secrets/HordeServer.Secrets.Tests/HordeServer.Secrets.Tests.csproj
--- a/Engine/Source/Programs/Horde/Plugins/Secrets/HordeServer.Secrets.Tests/HordeServer.Secrets.Tests.csproj
+++ b/Engine/Source/Programs/Horde/Plugins/Secrets/HordeServer.Secrets.Tests/HordeServer.Secrets.Tests.csproj
@@ -1,35 +1,13 @@
 <Project Sdk="Microsoft.NET.Sdk">
-  <Import Project="..\..\..\..\Shared\UnrealEngine.csproj.props" />
 
   <PropertyGroup>
-    <TargetFramework>net8.0</TargetFramework>
     <ImplicitUsings>enable</ImplicitUsings>
-    <Nullable>enable</Nullable>
-    <IsPackable>false</IsPackable>
-    <IsTestProject>true</IsTestProject>
     <RootNamespace>HordeServer.Tests</RootNamespace>
-    <RunAnalyzersDuringBuild>False</RunAnalyzersDuringBuild>
-    <Configurations>Debug;Release;Analyze</Configurations>
-    <NoWarn>$(NoWarn);1591</NoWarn>
-  </PropertyGroup>
-
-  <PropertyGroup Condition=" '$(Configuration)|$(Platform)' == 'Analyze|AnyCPU' ">
-    <RunAnalyzersDuringBuild>True</RunAnalyzersDuringBuild>
-    <TreatWarningsAsErrors>false</TreatWarningsAsErrors>
-    <EnforceCodeStyleInBuild>True</EnforceCodeStyleInBuild>
-    <GenerateDocumentationFile>True</GenerateDocumentationFile>
   </PropertyGroup>
 
   <ItemGroup>
     <PackageReference Include="AWSSDK.SSO" Version="3.7.501.12" />
     <PackageReference Include="AWSSDK.SSOOIDC" Version="3.7.502" />
-    <PackageReference Include="coverlet.collector" Version="6.0.4">
-      <PrivateAssets>all</PrivateAssets>
-      <IncludeAssets>runtime; build; native; contentfiles; analyzers; buildtransitive</IncludeAssets>
-    </PackageReference>
-    <PackageReference Include="Microsoft.NET.Test.Sdk" Version="17.14.1" />
-    <PackageReference Include="MSTest.TestAdapter" Version="3.11.1" />
-    <PackageReference Include="MSTest.TestFramework" Version="3.11.1" />
   </ItemGroup>
 
   <ItemGroup>
diff --git a/Engine/Source/Programs/Horde/Plugins/Secrets/HordeServer.Secrets/HordeServer.Secrets.csproj b/Engine/Source/Programs/Horde/Plugins/Secrets/HordeServer.Secrets/HordeServer.Secrets.csproj
--- a/Engine/Source/Programs/Horde/Plugins/Secrets/HordeServer.Secrets/HordeServer.Secrets.csproj
+++ b/Engine/Source/Programs/Horde/Plugins/Secrets/HordeServer.Secrets/HordeServer.Secrets.csproj
@@ -1,21 +1,8 @@
 <Project Sdk="Microsoft.NET.Sdk">
-  <Import Project="..\..\..\..\Shared\UnrealEngine.csproj.props" />
 
   <PropertyGroup>
-    <TargetFramework>net8.0</TargetFramework>
     <ImplicitUsings>enable</ImplicitUsings>
-    <Nullable>enable</Nullable>
     <RootNamespace>HordeServer</RootNamespace>
-    <RunAnalyzersDuringBuild>False</RunAnalyzersDuringBuild>
-    <Configurations>Debug;Release;Analyze</Configurations>
-  </PropertyGroup>
-
-  <PropertyGroup Condition=" '$(Configuration)|$(Platform)' == 'Analyze|AnyCPU' ">
-    <RunAnalyzersDuringBuild>True</RunAnalyzersDuringBuild>
-    <TreatWarningsAsErrors>false</TreatWarningsAsErrors>
-    <EnforceCodeStyleInBuild>True</EnforceCodeStyleInBuild>
-    <GenerateDocumentationFile>True</GenerateDocumentationFile>
-    <DocumentationFile></DocumentationFile>
   </PropertyGroup>
 
   <ItemGroup>
diff --git a/Engine/Source/Programs/Horde/Plugins/Storage/HordeServer.Storage.Tests/HordeServer.Storage.Tests.csproj b/Engine/Source/Programs/Horde/Plugins/Storage/HordeServer.Storage.Tests/HordeServer.Storage.Tests.csproj
--- a/Engine/Source/Programs/Horde/Plugins/Storage/HordeServer.Storage.Tests/HordeServer.Storage.Tests.csproj
+++ b/Engine/Source/Programs/Horde/Plugins/Storage/HordeServer.Storage.Tests/HordeServer.Storage.Tests.csproj
@@ -1,35 +1,10 @@
 <Project Sdk="Microsoft.NET.Sdk">
-  <Import Project="..\..\..\..\Shared\UnrealEngine.csproj.props" />
 
   <PropertyGroup>
-    <TargetFramework>net8.0</TargetFramework>
     <ImplicitUsings>enable</ImplicitUsings>
-    <Nullable>enable</Nullable>
-    <IsPackable>false</IsPackable>
-    <IsTestProject>true</IsTestProject>
     <RootNamespace>HordeServer.Tests</RootNamespace>
-    <RunAnalyzersDuringBuild>False</RunAnalyzersDuringBuild>
-    <Configurations>Debug;Release;Analyze</Configurations>
-    <NoWarn>$(NoWarn);1591</NoWarn>
   </PropertyGroup>
 
-  <PropertyGroup Condition=" '$(Configuration)|$(Platform)' == 'Analyze|AnyCPU' ">
-    <RunAnalyzersDuringBuild>True</RunAnalyzersDuringBuild>
-    <TreatWarningsAsErrors>false</TreatWarningsAsErrors>
-    <EnforceCodeStyleInBuild>True</EnforceCodeStyleInBuild>
-    <GenerateDocumentationFile>True</GenerateDocumentationFile>
-  </PropertyGroup>
-
-  <ItemGroup>
-    <PackageReference Include="coverlet.collector" Version="6.0.4">
-      <PrivateAssets>all</PrivateAssets>
-      <IncludeAssets>runtime; build; native; contentfiles; analyzers; buildtransitive</IncludeAssets>
-    </PackageReference>
-    <PackageReference Include="Microsoft.NET.Test.Sdk" Version="17.14.1" />
-    <PackageReference Include="MSTest.TestAdapter" Version="3.11.1" />
-    <PackageReference Include="MSTest.TestFramework" Version="3.11.1" />
-  </ItemGroup>
-
   <ItemGroup>
     <ProjectReference Include="..\HordeServer.Storage\HordeServer.Storage.csproj" />
     <ProjectReference Include="..\..\..\HordeServer.Tests.Shared\HordeServer.Tests.Shared.csproj" />
diff --git a/Engine/Source/Programs/Horde/Plugins/Storage/HordeServer.Storage/HordeServer.Storage.csproj b/Engine/Source/Programs/Horde/Plugins/Storage/HordeServer.Storage/HordeServer.Storage.csproj
--- a/Engine/Source/Programs/Horde/Plugins/Storage/HordeServer.Storage/HordeServer.Storage.csproj
+++ b/Engine/Source/Programs/Horde/Plugins/Storage/HordeServer.Storage/HordeServer.Storage.csproj
@@ -1,20 +1,7 @@
 ﻿<Project Sdk="Microsoft.NET.Sdk">
-  <Import Project="..\..\..\..\Shared\UnrealEngine.csproj.props" />
 
   <PropertyGroup>
-    <TargetFramework>net8.0</TargetFramework>
     <ImplicitUsings>enable</ImplicitUsings>
-    <Nullable>enable</Nullable>
-    <RunAnalyzersDuringBuild>False</RunAnalyzersDuringBuild>
-    <Configurations>Debug;Release;Analyze</Configurations>
-  </PropertyGroup>
-
-  <PropertyGroup Condition=" '$(Configuration)|$(Platform)' == 'Analyze|AnyCPU' ">
-    <RunAnalyzersDuringBuild>True</RunAnalyzersDuringBuild>
-    <TreatWarningsAsErrors>false</TreatWarningsAsErrors>
-    <EnforceCodeStyleInBuild>True</EnforceCodeStyleInBuild>
-    <GenerateDocumentationFile>True</GenerateDocumentationFile>
-    <DocumentationFile></DocumentationFile>
   </PropertyGroup>
 
   <ItemGroup>
diff --git a/Engine/Source/Programs/Horde/Plugins/Symbols/HordeServer.Symbols/HordeServer.Symbols.csproj b/Engine/Source/Programs/Horde/Plugins/Symbols/HordeServer.Symbols/HordeServer.Symbols.csproj
--- a/Engine/Source/Programs/Horde/Plugins/Symbols/HordeServer.Symbols/HordeServer.Symbols.csproj
+++ b/Engine/Source/Programs/Horde/Plugins/Symbols/HordeServer.Symbols/HordeServer.Symbols.csproj
@@ -1,21 +1,8 @@
 <Project Sdk="Microsoft.NET.Sdk">
-  <Import Project="..\..\..\..\Shared\UnrealEngine.csproj.props" />
 
   <PropertyGroup>
-    <TargetFramework>net8.0</TargetFramework>
     <ImplicitUsings>enable</ImplicitUsings>
-    <Nullable>enable</Nullable>
     <RootNamespace>HordeServer</RootNamespace>
-    <RunAnalyzersDuringBuild>False</RunAnalyzersDuringBuild>
-    <Configurations>Debug;Release;Analyze</Configurations>
-  </PropertyGroup>
-
-  <PropertyGroup Condition=" '$(Configuration)|$(Platform)' == 'Analyze|AnyCPU' ">
-    <RunAnalyzersDuringBuild>True</RunAnalyzersDuringBuild>
-    <TreatWarningsAsErrors>false</TreatWarningsAsErrors>
-    <EnforceCodeStyleInBuild>True</EnforceCodeStyleInBuild>
-    <GenerateDocumentationFile>True</GenerateDocumentationFile>
-    <DocumentationFile></DocumentationFile>
   </PropertyGroup>
 
   <ItemGroup>
diff --git a/Engine/Source/Programs/Horde/Plugins/Tools/HordeServer.Tools.Tests/HordeServer.Tools.Tests.csproj b/Engine/Source/Programs/Horde/Plugins/Tools/HordeServer.Tools.Tests/HordeServer.Tools.Tests.csproj
--- a/Engine/Source/Programs/Horde/Plugins/Tools/HordeServer.Tools.Tests/HordeServer.Tools.Tests.csproj
+++ b/Engine/Source/Programs/Horde/Plugins/Tools/HordeServer.Tools.Tests/HordeServer.Tools.Tests.csproj
@@ -1,34 +1,10 @@
 <Project Sdk="Microsoft.NET.Sdk">
-  <Import Project="..\..\..\..\Shared\UnrealEngine.csproj.props" />
 
   <PropertyGroup>
-    <TargetFramework>net8.0</TargetFramework>
     <ImplicitUsings>enable</ImplicitUsings>
-    <Nullable>enable</Nullable>
-    <IsPackable>false</IsPackable>
-    <IsTestProject>true</IsTestProject>
     <RootNamespace>HordeServer.Tests</RootNamespace>
-    <RunAnalyzersDuringBuild>False</RunAnalyzersDuringBuild>
-    <Configurations>Debug;Release;Analyze</Configurations>
-    <NoWarn>$(NoWarn);1591</NoWarn>
   </PropertyGroup>
 
-  <PropertyGroup Condition=" '$(Configuration)|$(Platform)' == 'Analyze|AnyCPU' ">
-    <RunAnalyzersDuringBuild>True</RunAnalyzersDuringBuild>
-    <TreatWarningsAsErrors>false</TreatWarningsAsErrors>
-    <EnforceCodeStyleInBuild>True</EnforceCodeStyleInBuild>
-    <GenerateDocumentationFile>True</GenerateDocumentationFile>
-  </PropertyGroup>
-
-  <ItemGroup>
-    <PackageReference Include="coverlet.collector" Version="6.0.4">
-      <PrivateAssets>all</PrivateAssets>
-      <IncludeAssets>runtime; build; native; contentfiles; analyzers; buildtransitive</IncludeAssets>
-    </PackageReference>
-    <PackageReference Include="Microsoft.NET.Test.Sdk" Version="17.14.1" />
-    <PackageReference Include="MSTest.TestAdapter" Version="3.11.1" />
-    <PackageReference Include="MSTest.TestFramework" Version="3.11.1" />
-  </ItemGroup>
 
   <ItemGroup>
     <ProjectReference Include="..\HordeServer.Tools\HordeServer.Tools.csproj" />
diff --git a/Engine/Source/Programs/Horde/Plugins/Tools/HordeServer.Tools/HordeServer.Tools.csproj b/Engine/Source/Programs/Horde/Plugins/Tools/HordeServer.Tools/HordeServer.Tools.csproj
--- a/Engine/Source/Programs/Horde/Plugins/Tools/HordeServer.Tools/HordeServer.Tools.csproj
+++ b/Engine/Source/Programs/Horde/Plugins/Tools/HordeServer.Tools/HordeServer.Tools.csproj
@@ -1,21 +1,8 @@
 ﻿<Project Sdk="Microsoft.NET.Sdk">
-  <Import Project="..\..\..\..\Shared\UnrealEngine.csproj.props" />
 
   <PropertyGroup>
-    <TargetFramework>net8.0</TargetFramework>
     <ImplicitUsings>enable</ImplicitUsings>
-    <Nullable>enable</Nullable>
     <RootNamespace>HordeServer</RootNamespace>
-    <RunAnalyzersDuringBuild>False</RunAnalyzersDuringBuild>
-    <Configurations>Debug;Release;Analyze</Configurations>
-  </PropertyGroup>
-
-  <PropertyGroup Condition=" '$(Configuration)|$(Platform)' == 'Analyze|AnyCPU' ">
-    <RunAnalyzersDuringBuild>True</RunAnalyzersDuringBuild>
-    <TreatWarningsAsErrors>false</TreatWarningsAsErrors>
-    <EnforceCodeStyleInBuild>True</EnforceCodeStyleInBuild>
-    <GenerateDocumentationFile>True</GenerateDocumentationFile>
-    <DocumentationFile></DocumentationFile>
   </PropertyGroup>
 
   <ItemGroup>
diff --git a/Engine/Source/Programs/Horde/Samples/RemoteClient/RemoteClient.csproj b/Engine/Source/Programs/Horde/Samples/RemoteClient/RemoteClient.csproj
--- a/Engine/Source/Programs/Horde/Samples/RemoteClient/RemoteClient.csproj
+++ b/Engine/Source/Programs/Horde/Samples/RemoteClient/RemoteClient.csproj
@@ -1,11 +1,8 @@
 <Project Sdk="Microsoft.NET.Sdk">
-  <Import Project="..\..\..\Shared\UnrealEngine.csproj.props" />
 
   <PropertyGroup>
     <OutputType>Exe</OutputType>
-    <TargetFramework>net8.0</TargetFramework>
     <ImplicitUsings>enable</ImplicitUsings>
-    <Nullable>enable</Nullable>
   </PropertyGroup>
 
   <ItemGroup>
diff --git a/Engine/Source/Programs/Horde/Samples/RemoteWorker/RemoteWorker.csproj b/Engine/Source/Programs/Horde/Samples/RemoteWorker/RemoteWorker.csproj
--- a/Engine/Source/Programs/Horde/Samples/RemoteWorker/RemoteWorker.csproj
+++ b/Engine/Source/Programs/Horde/Samples/RemoteWorker/RemoteWorker.csproj
@@ -1,11 +1,8 @@
 <Project Sdk="Microsoft.NET.Sdk">
-  <Import Project="..\..\..\Shared\UnrealEngine.csproj.props" />
 
   <PropertyGroup>
     <OutputType>Exe</OutputType>
-    <TargetFramework>net8.0</TargetFramework>
     <ImplicitUsings>enable</ImplicitUsings>
-    <Nullable>enable</Nullable>
   </PropertyGroup>
 
   <ItemGroup>
diff --git a/Engine/Source/Programs/Horde/Tools/ReleaseNotes/ReleaseNotes.csproj b/Engine/Source/Programs/Horde/Tools/ReleaseNotes/ReleaseNotes.csproj
--- a/Engine/Source/Programs/Horde/Tools/ReleaseNotes/ReleaseNotes.csproj
+++ b/Engine/Source/Programs/Horde/Tools/ReleaseNotes/ReleaseNotes.csproj
@@ -1,11 +1,8 @@
 <Project Sdk="Microsoft.NET.Sdk">
-  <Import Project="..\..\..\Shared\UnrealEngine.csproj.props" />
 
   <PropertyGroup>
     <OutputType>Exe</OutputType>
-    <TargetFramework>net8.0</TargetFramework>
     <ImplicitUsings>enable</ImplicitUsings>
-    <Nullable>enable</Nullable>
     <PublishSingleFile>true</PublishSingleFile>
   </PropertyGroup>
 
