Horde: Updating schema documentation

#rnx

diff --git a/Engine/Source/Programs/Horde/Docs/Config/Schema/AclActions.md b/Engine/Source/Programs/Horde/Docs/Config/Schema/AclActions.md
--- a/Engine/Source/Programs/Horde/Docs/Config/Schema/AclActions.md
+++ b/Engine/Source/Programs/Horde/Docs/Config/Schema/AclActions.md
@@ -37,48 +37,6 @@
 | `DeleteAccount` |  |
 | `ViewAccount` |  |
 
-## Artifacts
-
-| Name | Description |
-| ---- | ----------- |
-| `ReadArtifact` | Permission to read from an artifact |
-| `WriteArtifact` | Permission to write to an artifact |
-| `DeleteArtifact` | Permission to delete to an artifact |
-| `UploadArtifact` | Ability to create an artifact. Typically just for debugging; agents have this access for a particular session. |
-| `DownloadArtifact` | Ability to download an artifact |
-
-## Bisect
-
-| Name | Description |
-| ---- | ----------- |
-| `CreateBisectTask` | Ability to start new bisect tasks |
-| `UpdateBisectTask` | Ability to update a bisect task |
-| `ViewBisectTask` | Ability to view a bisect task |
-
-## BuildHealth
-
-| Name | Description |
-| ---- | ----------- |
-| `AddUpdateDeleteBuildHealthFilter` | Ability to add, update, and delete build health filters. |
-
-## Devices
-
-| Name | Description |
-| ---- | ----------- |
-| `DeviceRead` | Ability to read devices |
-| `DeviceWrite` | Ability to write devices |
-
-## Jobs
-
-| Name | Description |
-| ---- | ----------- |
-| `CreateJob` | Ability to start new jobs |
-| `UpdateJob` | Rename a job, modify its priority, etc... |
-| `DeleteJob` | Delete a job properties |
-| `ExecuteJob` | Allows updating a job metadata (name, changelist number, step properties, new groups, job states, etc...). Typically granted to agents. Not user facing. |
-| `RetryJobStep` | Ability to retry a failed job step |
-| `ViewJob` | Ability to view a job |
-
 ## Notices
 
 | Name | Description |
@@ -87,28 +45,6 @@
 | `UpdateNotice` |  |
 | `DeleteNotice` |  |
 
-## Notifications
-
-| Name | Description |
-| ---- | ----------- |
-| `CreateSubscription` | Ability to subscribe to notifications |
-
-## Projects
-
-| Name | Description |
-| ---- | ----------- |
-| `CreateProject` | Allows the creation of new projects |
-| `DeleteProject` | Allows deletion of projects. |
-| `UpdateProject` | Modify attributes of a project (name, categories, etc...) |
-| `ViewProject` | View information about a project |
-
-## Replicators
-
-| Name | Description |
-| ---- | ----------- |
-| `UpdateReplicator` | Allows deletion of projects. |
-| `ViewReplicator` | Allows the creation of new projects |
-
 ## ServiceAccounts
 
 | Name | Description |
@@ -118,24 +54,6 @@
 | `DeleteAccount` |  |
 | `ViewAccount` |  |
 
-## Streams
-
-| Name | Description |
-| ---- | ----------- |
-| `CreateStream` | Allows the creation of new streams within a project |
-| `UpdateStream` | Allows updating a stream (agent types, templates, schedules) |
-| `DeleteStream` | Allows deleting a stream |
-| `ViewStream` | Ability to view a stream |
-| `ViewChanges` | View changes submitted to a stream. NOTE: this returns responses from the server's Perforce account, which may be a priviledged user. |
-| `ViewTemplate` | View template associated with a stream |
-
-## TestData
-
-| Name | Description |
-| ---- | ----------- |
-| `TestHealthRead` | Ability to read test health |
-| `TestHealthWrite` | Ability to write test health |
-
 ## Accounts
 
 | Name | Description |
@@ -145,17 +63,6 @@
 | `DeleteAccount` |  |
 | `ViewAccount` |  |
 
-## Agents
-
-| Name | Description |
-| ---- | ----------- |
-| `CreateAgent` | Ability to create an agent. This may be done explicitly, or granted to agents to allow them to self-register. |
-| `CreateWorkstationAgent` | Ability to create a workstation agent. Granted to users with hardware to offer. |
-| `UpdateAgent` | Update an agent's name, pools, etc... |
-| `DeleteAgent` | Soft-delete an agent |
-| `ViewAgent` | View an agent |
-| `ListAgents` | List the available agents |
-
 ## Artifacts
 
 | Name | Description |
@@ -180,15 +87,6 @@
 | ---- | ----------- |
 | `AddUpdateDeleteBuildHealthFilter` |  |
 
-## Compute
-
-| Name | Description |
-| ---- | ----------- |
-| `AddComputeTasks` | User can add tasks to the compute cluster |
-| `GetComputeTasks` | User can get and list tasks from the compute cluster |
-| `UbaCacheRead` | User can read from Unreal Build Accelerator cache |
-| `UbaCacheWrite` | User can write to Unreal Build Accelerator cache |
-
 ## Devices
 
 | Name | Description |
@@ -207,24 +105,6 @@
 | `RetryJobStep` |  |
 | `ViewJob` |  |
 
-## Leases
-
-| Name | Description |
-| ---- | ----------- |
-| `ViewLeases` | View all the leases that an agent has worked on |
-| `ViewLeaseTasks` | View the task data for a lease |
-
-## Logs
-
-| Name | Description |
-| ---- | ----------- |
-| `CreateLog` | Ability to create a log. Implicitly granted to agents. |
-| `UpdateLog` | Ability to update log metadata |
-| `ViewLog` | Ability to view a log contents |
-| `WriteLogData` | Ability to write log data |
-| `CreateEvent` | Ability to create events |
-| `ViewEvent` | Ability to view events |
-
 ## Notices
 
 | Name | Description |
@@ -239,16 +119,6 @@
 | ---- | ----------- |
 | `CreateSubscription` |  |
 
-## Pools
-
-| Name | Description |
-| ---- | ----------- |
-| `CreatePool` | Create a global pool of agents |
-| `UpdatePool` | Modify an agent pool |
-| `DeletePool` | Delete an agent pool |
-| `ViewPool` | Ability to view a pool |
-| `ListPools` | View all the available agent pools |
-
 ## Projects
 
 | Name | Description |
@@ -274,21 +144,6 @@
 | `DeleteAccount` |  |
 | `ViewAccount` |  |
 
-## Sessions
-
-| Name | Description |
-| ---- | ----------- |
-| `CreateSession` | Granted to agents to call CreateSession, which returns a bearer token identifying themselves valid to call UpdateSesssion via gRPC. |
-| `ViewSession` | Allows viewing information about an agent session |
-
-## Software
-
-| Name | Description |
-| ---- | ----------- |
-| `UploadSoftware` | Ability to upload new versions of the agent software |
-| `DownloadSoftware` | Ability to download the agent software |
-| `DeleteSoftware` | Ability to delete agent software |
-
 ## Streams
 
 | Name | Description |
@@ -782,7 +637,7 @@
 
 | Name | Description |
 | ---- | ----------- |
-| `ViewSecret` | View a credential |
+| `ViewSecret` |  |
 
 ## ServiceAccounts
 
@@ -985,6 +840,19 @@
 | `DownloadSoftware` |  |
 | `DeleteSoftware` |  |
 
+## Storage
+
+| Name | Description |
+| ---- | ----------- |
+| `ReadBlobs` |  |
+| `WriteBlobs` |  |
+| `ReadAliases` |  |
+| `WriteAliases` |  |
+| `DeleteAliases` |  |
+| `ReadRefs` |  |
+| `WriteRefs` |  |
+| `DeleteRefs` |  |
+
 ## Streams
 
 | Name | Description |
@@ -1166,14 +1034,14 @@
 
 | Name | Description |
 | ---- | ----------- |
-| `ReadBlobs` | Ability to read blobs from the storage service |
-| `WriteBlobs` | Ability to write blobs to the storage service |
-| `ReadAliases` | Ability to read aliases from the storage service |
-| `WriteAliases` | Ability to write aliases to the storage service |
-| `DeleteAliases` | Ability to write aliases to the storage service |
-| `ReadRefs` | Ability to read refs from the storage service |
-| `WriteRefs` | Ability to write refs to the storage service |
-| `DeleteRefs` | Ability to delete refs |
+| `ReadBlobs` |  |
+| `WriteBlobs` |  |
+| `ReadAliases` |  |
+| `WriteAliases` |  |
+| `DeleteAliases` |  |
+| `ReadRefs` |  |
+| `WriteRefs` |  |
+| `DeleteRefs` |  |
 
 ## Streams
 
@@ -1376,18 +1244,19 @@
 | `ViewChanges` |  |
 | `ViewTemplate` |  |
 
-## Symbols
+## TestData
 
 | Name | Description |
 | ---- | ----------- |
-| `ReadSymbols` | Ability to download symbols |
+| `TestHealthRead` |  |
+| `TestHealthWrite` |  |
 
-## TestData
+## Tools
 
 | Name | Description |
 | ---- | ----------- |
-| `TestHealthRead` |  |
-| `TestHealthWrite` |  |
+| `DownloadTool` |  |
+| `UploadTool` |  |
 
 ## Accounts
 
@@ -1589,8 +1458,8 @@
 
 | Name | Description |
 | ---- | ----------- |
-| `DownloadTool` | Ability to download a tool |
-| `UploadTool` | Ability to upload new tool versions |
+| `DownloadTool` |  |
+| `UploadTool` |  |
 
 ## Accounts
 
diff --git a/Engine/Source/Programs/Horde/Docs/Config/Schema/Dashboard.md b/Engine/Source/Programs/Horde/Docs/Config/Schema/Dashboard.md
--- a/Engine/Source/Programs/Horde/Docs/Config/Schema/Dashboard.md
+++ b/Engine/Source/Programs/Horde/Docs/Config/Schema/Dashboard.md
@@ -39,17 +39,13 @@ Name | Description
 
 ## ConfigInclude
 
-Directive to merge config data from another source
-
 Name | Description
 ---- | -----------
-`path` | `string`<br>Path to the config data to be included. May be relative to the including file's location.
+`path` | `string`<br>
 
 ## ConfigMacro
 
-Declares a config macro
-
 Name | Description
 ---- | -----------
-`name` | `string`<br>Name of the macro property
-`value` | `string`<br>Value for the macro property
+`name` | `string`<br>
+`value` | `string`<br>
diff --git a/Engine/Source/Programs/Horde/Docs/Config/Schema/Globals.md b/Engine/Source/Programs/Horde/Docs/Config/Schema/Globals.md
--- a/Engine/Source/Programs/Horde/Docs/Config/Schema/Globals.md
+++ b/Engine/Source/Programs/Horde/Docs/Config/Schema/Globals.md
@@ -18,20 +18,16 @@ Name | Description
 
 ## ConfigInclude
 
-Directive to merge config data from another source
-
 Name | Description
 ---- | -----------
-`path` | `string`<br>Path to the config data to be included. May be relative to the including file's location.
+`path` | `string`<br>
 
 ## ConfigMacro
 
-Declares a config macro
-
 Name | Description
 ---- | -----------
-`name` | `string`<br>Name of the macro property
-`value` | `string`<br>Value for the macro property
+`name` | `string`<br>
+`value` | `string`<br>
 
 ## ScheduledDowntime
 
@@ -59,127 +55,117 @@ Name | Description
 
 Name | Description
 ---- | -----------
+`analytics` | [AnalyticsConfig](#analyticsconfig)<br>Configuration for the analytics plugin
 `build` | [BuildConfig](#buildconfig)<br>Configuration for the build plugin
 `compute` | [ComputeConfig](#computeconfig)<br>Configuration for the compute plugin
-`analytics` | [AnalyticsConfig](#analyticsconfig)<br>Configuration for the analytics plugin
-`exampleplugin` | [ExamplePluginConfig](#examplepluginconfig)<br>Configuration for the exampleplugin plugin
-`secrets` | [SecretsConfig](#secretsconfig)<br>Configuration for the secrets plugin
 `experimental` | [ExperimentalConfig](#experimentalconfig)<br>Configuration for the experimental plugin
+`secrets` | [SecretsConfig](#secretsconfig)<br>Configuration for the secrets plugin
 `storage` | [StorageConfig](#storageconfig)<br>Configuration for the storage plugin
-`symbols` | [SymbolsConfig](#symbolsconfig)<br>Configuration for the symbols plugin
+`exampleplugin` | [ExamplePluginConfig](#examplepluginconfig)<br>Configuration for the exampleplugin plugin
 `tools` | [ToolsConfig](#toolsconfig)<br>Configuration for the tools plugin
+`symbols` | [SymbolsConfig](#symbolsconfig)<br>Configuration for the symbols plugin
+
+## AnalyticsConfig
+
+Name | Description
+---- | -----------
+`stores` | [TelemetryStoreConfig](Telemetry.md)`[]`<br>
 
 ## BuildConfig
 
-Configuration for the build plugin
-
-Name | Description
----- | -----------
-`perforceClusters` | [PerforceCluster](#perforcecluster)`[]`<br>List of Perforce clusters
-`devices` | [DeviceConfig](#deviceconfig)<br>Device configuration
-`maxConformCount` | `integer`<br>Maximum number of conforms to run at once
-`agentShutdownIfDisabledGracePeriod` | `string`<br>Time to wait before shutting down an agent that has been disabled Used if no value is set on the actual pool.
-`artifactTypes` | [ArtifactTypeConfig](#artifacttypeconfig)`[]`<br>Configuration for different artifact types
-`projects` | [ProjectConfig](Projects.md)`[]`<br>List of projects
-`enableConformTasks` | `boolean`<br>Whether to allow conform tasks to run
-`enableJobTaskConformCheck` | `boolean`<br>Whether to check if a conform is required on an agent before assigning a lease for a job task source.
-`issueFixedTag` | `string`<br>Commit tag to use for marking issues as fixed
-`issues` | [IssueConfig](Issues.md)<br>Build health issue configurations
-`enableIssueRootCauseAnalysisNotifications` | `boolean`<br>Whether to allow notifications for Root Cause Analysis to be sent
-`autoRcaUseMultiStage` | `boolean`<br>Whether to use multi-stage LLM analysis for auto-RCA. Multi-stage is slower but potentially more accurate. Set to false for faster processing with higher throughput.
-`maxConcurrentRcaRequests` | `integer`<br>Maximum number of concurrent auto-RCA requests. Limits load on the external LLM API. Set to 0 to disable auto-RCA globally (acts as a kill switch).
-`issueFindSuspectsMaxSuccessDays` | `integer`<br>The maximum days since node success to scan for suspects
-`issueFindSuspectsMaxScoringSeconds` | `integer`<br>The maximum number of seconds to spend scording suspects
-`issueFindSuspectsMaxChanges` | `integer`<br>Maximum number of commits to examine for suspects. Previously hardcoded as const MaxChanges = 500.
-`issueFindSuspectsEnableParallelScoring` | `boolean`<br>Enable parallel scoring of commits for suspect finding. When enabled, commits are scored concurrently up to MaxParallelism.
-`issueFindSuspectsMaxParallelism` | `integer`<br>Maximum concurrent commit file lookups when parallel scoring is enabled. Higher values improve speed but increase Perforce server load.
-`issueFindSuspectsEnableEarlyExit` | `boolean`<br>Enable early exit when theoretical maximum score is achieved. Only exits early if the best possible score has been found.
-`issueFindSuspectsEnableBidirectionalQuery` | `boolean`<br>Enable bidirectional commit query to capture both recent and old commits. Queries half the limit from each end of the commit range.
-`issueFindSuspectsEnableFingerprintFallback` | `boolean`<br>Enable fingerprint-based fallback when no suspects are found. Helps prevent duplicate issues when commit range exceeds MaxChanges.
-`issueFingerprintFallbackMinScore` | `integer`<br>Minimum score required for fingerprint fallback matching. Higher values are more conservative (fewer false matches).
-`issueFingerprintFallbackMaxAgeDays` | `integer`<br>Maximum age in days for issues to be considered in fingerprint fallback.
-`issueFingerprintFallbackExcludedTypes` | `string[]`<br>Fingerprint types excluded from fallback matching (too generic).
-`enableExpandedIgnorePatterns` | `boolean`<br>Whether to enable expanded IgnorePatterns discovery (workspace-wide scanning + config-specified files). When false (default), only Engine directory is scanned for ignore patterns.
+Name | Description
+---- | -----------
+`perforceClusters` | [PerforceCluster](#perforcecluster)`[]`<br>
+`devices` | [DeviceConfig](#deviceconfig)<br>
+`maxConformCount` | `integer`<br>
+`agentShutdownIfDisabledGracePeriod` | `string`<br>
+`artifactTypes` | [ArtifactTypeConfig](#artifacttypeconfig)`[]`<br>
+`projects` | [ProjectConfig](Projects.md)`[]`<br>
+`enableConformTasks` | `boolean`<br>
+`enableJobTaskConformCheck` | `boolean`<br>
+`issueFixedTag` | `string`<br>
+`issues` | [IssueConfig](Issues.md)<br>
+`enableIssueRootCauseAnalysisNotifications` | `boolean`<br>
+`autoRcaUseMultiStage` | `boolean`<br>
+`maxConcurrentRcaRequests` | `integer`<br>
+`issueFindSuspectsMaxSuccessDays` | `integer`<br>
+`issueFindSuspectsMaxScoringSeconds` | `integer`<br>
+`issueFindSuspectsMaxChanges` | `integer`<br>
+`issueFindSuspectsEnableParallelScoring` | `boolean`<br>
+`issueFindSuspectsMaxParallelism` | `integer`<br>
+`issueFindSuspectsEnableEarlyExit` | `boolean`<br>
+`issueFindSuspectsEnableBidirectionalQuery` | `boolean`<br>
+`issueFindSuspectsEnableFingerprintFallback` | `boolean`<br>
+`issueFingerprintFallbackMinScore` | `integer`<br>
+`issueFingerprintFallbackMaxAgeDays` | `integer`<br>
+`issueFingerprintFallbackExcludedTypes` | `string[]`<br>
+`enableExpandedIgnorePatterns` | `boolean`<br>
 
 ## PerforceCluster
 
-Information about a cluster of Perforce servers.
-
 Name | Description
 ---- | -----------
-`name` | `string`<br>Name of the cluster
-`serviceAccount` | `string`<br>Username for Horde to log in to this server. Will use the first account specified below if not overridden.
-`canImpersonate` | `boolean`<br>Whether the service account can impersonate other users
-`supportsPartitionedWorkspaces` | `boolean`<br>Whether to use partitioned workspaces on this server
-`servers` | [PerforceServer](#perforceserver)`[]`<br>List of servers
-`credentials` | [PerforceCredentials](#perforcecredentials)`[]`<br>List of server credentials
-`autoSdk` | [AutoSdkWorkspace](#autosdkworkspace)`[]`<br>List of autosdk streams
+`name` | `string`<br>
+`serviceAccount` | `string`<br>
+`canImpersonate` | `boolean`<br>
+`supportsPartitionedWorkspaces` | `boolean`<br>
+`servers` | [PerforceServer](#perforceserver)`[]`<br>
+`credentials` | [PerforceCredentials](#perforcecredentials)`[]`<br>
+`autoSdk` | [AutoSdkWorkspace](#autosdkworkspace)`[]`<br>
 
 ## PerforceServer
 
-Information about an individual Perforce server
-
 Name | Description
 ---- | -----------
-`serverAndPort` | `string`<br>The server and port. The server may be a DNS entry with multiple records, in which case it will be actively load balanced. If "ssl:" prefix is used, ensure P4 server's fingerprint/certificate is trusted. See Horde's documentation on connecting to SSL-enabled Perforce servers.
-`healthCheck` | `boolean`<br>Whether to query the healthcheck address under each server
-`resolveDns` | `boolean`<br>Whether to resolve the DNS entries and load balance between different hosts
-`maxConformCount` | `integer`<br>Maximum number of simultaneous conforms on this server
-`condition` | `string`<br>Optional condition for a machine to be eligible to use this server
-`properties` | `string[]`<br>List of properties for an agent to be eligible to use this server
+`serverAndPort` | `string`<br>
+`healthCheck` | `boolean`<br>
+`resolveDns` | `boolean`<br>
+`maxConformCount` | `integer`<br>
+`condition` | `string`<br>
+`properties` | `string[]`<br>
 
 ## PerforceCredentials
 
-Credentials for a Perforce user
-
 Name | Description
 ---- | -----------
-`userName` | `string`<br>The username
-`password` | `string`<br>Password for the user
-`ticket` | `string`<br>Login ticket for the user (will be used instead of password if set)
+`userName` | `string`<br>
+`password` | `string`<br>
+`ticket` | `string`<br>
 
 ## AutoSdkWorkspace
 
-Path to a platform and stream to use for syncing AutoSDK
-
 Name | Description
 ---- | -----------
-`name` | `string`<br>Name of this workspace
-`properties` | `string[]`<br>The agent properties to check (eg. "OSFamily=Windows")
-`userName` | `string`<br>Username for logging in to the server
-`stream` | `string`<br>Stream to use
+`name` | `string`<br>
+`properties` | `string[]`<br>
+`userName` | `string`<br>
+`stream` | `string`<br>
 
 ## DeviceConfig
 
-Configuration for devices
-
 Name | Description
 ---- | -----------
-`platforms` | [DevicePlatformConfig](#deviceplatformconfig)`[]`<br>List of device platforms
-`pools` | [DevicePoolConfig](#devicepoolconfig)`[]`<br>List of device pools
+`platforms` | [DevicePlatformConfig](#deviceplatformconfig)`[]`<br>
+`pools` | [DevicePoolConfig](#devicepoolconfig)`[]`<br>
 
 ## DevicePlatformConfig
 
-Configuration for a device platform
-
 Name | Description
 ---- | -----------
-`id` | `string`<br>The id for this platform
-`name` | `string`<br>Name of the platform
-`models` | `string[]`<br>A list of platform models
-`legacyNames` | `string[]`<br>Legacy names which older versions of Gauntlet may be using
-`legacyPerfSpecHighModel` | `string`<br>Model name for the high perf spec, which may be requested by Gauntlet
+`id` | `string`<br>
+`name` | `string`<br>
+`models` | `string[]`<br>
+`legacyNames` | `string[]`<br>
+`legacyPerfSpecHighModel` | `string`<br>
 
 ## DevicePoolConfig
 
-Configuration for a device pool
-
 Name | Description
 ---- | -----------
-`id` | `string`<br>The id for this platform
-`name` | `string`<br>The name of the pool
-`poolType` | [DevicePoolType](#devicepooltype-enum)<br>The type of the pool
-`projectIds` | `string[]`<br>List of project ids associated with pool
+`id` | `string`<br>
+`name` | `string`<br>
+`poolType` | [DevicePoolType](#devicepooltype-enum)<br>
+`projectIds` | `string[]`<br>
 
 ## DevicePoolType (Enum)
 
@@ -192,149 +178,125 @@ Name | Description
 
 ## ArtifactTypeConfig
 
-Configuration for an artifact
-
 Name | Description
 ---- | -----------
-`name` | `string`<br>Legacy 'Name' property
-`type` | `string`<br>Name of the artifact type
-`acl` | [AclConfig](#aclconfig)<br>Acl for the artifact type
-`keepCount` | `integer`<br>Number of artifacts to retain
-`keepDays` | `integer`<br>Number of days to retain artifacts of this type
-`namespaceId` | `string`<br>Storage namespace to use for this artifact types
+`name` | `string`<br>
+`type` | `string`<br>
+`acl` | [AclConfig](#aclconfig)<br>
+`keepCount` | `integer`<br>
+`keepDays` | `integer`<br>
+`namespaceId` | `string`<br>
 
 ## AclConfig
 
-Parameters to update an ACL
-
 Name | Description
 ---- | -----------
-`entries` | [AclEntryConfig](#aclentryconfig)`[]`<br>Entries to replace the existing ACL
-`profiles` | [AclProfileConfig](#aclprofileconfig)`[]`<br>Defines profiles which allow grouping sets of actions into named collections
-`inherit` | `boolean`<br>Whether to inherit permissions from the parent ACL
-`exceptions` | `string[]`<br>List of exceptions to the inherited setting
+`entries` | [AclEntryConfig](#aclentryconfig)`[]`<br>
+`profiles` | [AclProfileConfig](#aclprofileconfig)`[]`<br>
+`inherit` | `boolean`<br>
+`exceptions` | `string[]`<br>
 
 ## AclEntryConfig
 
-Individual entry in an ACL
-
 Name | Description
 ---- | -----------
-`claim` | [AclClaimConfig](#aclclaimconfig)<br>Name of the user or group
-`actions` | `string[]`<br>Array of actions to allow
-`profiles` | `string[]`<br>List of profiles to grant
+`claim` | [AclClaimConfig](#aclclaimconfig)<br>
+`actions` | `string[]`<br>
+`profiles` | `string[]`<br>
 
 ## AclClaimConfig
 
-New claim to create
-
 Name | Description
 ---- | -----------
-`type` | `string`<br>The claim type
-`value` | `string`<br>The claim value
+`type` | `string`<br>
+`value` | `string`<br>
 
 ## AclProfileConfig
 
-Configuration for an ACL profile. This defines a preset group of actions which can be given to a user via an ACL entry.
-
 Name | Description
 ---- | -----------
-`id` | `string`<br>Identifier for this profile
-`actions` | `string[]`<br>Actions to include
-`excludeActions` | `string[]`<br>Actions to exclude from the inherited actions
-`extends` | `string[]`<br>Other profiles to extend from
+`id` | `string`<br>
+`actions` | `string[]`<br>
+`excludeActions` | `string[]`<br>
+`extends` | `string[]`<br>
 
 ## ComputeConfig
 
-Configuration for the compute system
-
 Name | Description
 ---- | -----------
-`acl` | [AclConfig](#aclconfig)<br>Inherited root acl
-`versionEnum` | [ConfigVersion](#configversion-enum)<br>Config version number
-`telemetryStoreId` | `string`<br>Telemetry store for lease and compute telemetry data
-`rates` | [AgentRateConfig](#agentrateconfig)`[]`<br>List of costs of a particular agent type
-`clusters` | [ComputeClusterConfig](#computeclusterconfig)`[]`<br>List of compute profiles
-`pools` | [PoolConfig](#poolconfig)`[]`<br>List of pools
-`software` | [AgentSoftwareConfig](#agentsoftwareconfig)`[]`<br>List of costs of a particular agent type
-`networks` | [NetworkConfig](#networkconfig)`[]`<br>List of networks
+`acl` | [AclConfig](#aclconfig)<br>
+`versionEnum` | [ConfigVersion](#configversion-enum)<br>
+`telemetryStoreId` | `string`<br>
+`rates` | [AgentRateConfig](#agentrateconfig)`[]`<br>
+`clusters` | [ComputeClusterConfig](#computeclusterconfig)`[]`<br>
+`pools` | [PoolConfig](#poolconfig)`[]`<br>
+`software` | [AgentSoftwareConfig](#agentsoftwareconfig)`[]`<br>
+`networks` | [NetworkConfig](#networkconfig)`[]`<br>
 
 ## ConfigVersion (Enum)
 
-Global version number for running the server. As new features are introduced that require data migrations, this version number indicates the backwards compatibility functionality that must be enabled. When adding a new version here, also add a message to ConfigService.CreateSnapshotAsync describing the steps that need to be taken to upgrade the deployment.
-
 Name | Description
 ---- | -----------
-`None` | Not specified
-`Initial` | Initial version number
-`PoolsInConfigFiles` | Ability to add/remove pools via the REST API is removed. Pools should be configured through globals.json instead.
-`IncrementalPoolMigration` | Safe hybrid mode for incremental pool migration. Pools are read from both config files and the database, with deduplication. Config file pools take priority — database pools with matching IDs are skipped. The pool editor UI remains visible for database-managed pools. Use /api/v1/server/migrate/pool-config to export remaining database pools.
-`Latest` | Latest version number
-`LatestPlusOne` | One after the last defined version number
+`None` | 
+`Initial` | 
+`PoolsInConfigFiles` | 
+`IncrementalPoolMigration` | 
+`Latest` | 
+`LatestPlusOne` | 
 
 ## AgentRateConfig
 
-Describes the monetary cost of agents matching a particular criteria
-
 Name | Description
 ---- | -----------
-`condition` | `string`<br>Condition string
-`rate` | `number`<br>Rate for this agent
+`condition` | `string`<br>
+`rate` | `number`<br>
 
 ## ComputeClusterConfig
 
-Profile for executing compute requests
-
 Name | Description
 ---- | -----------
-`id` | `string`<br>Name of the partition
-`namespaceId` | `string`<br>Name of the namespace to use
-`requestBucketId` | `string`<br>Name of the input bucket
-`responseBucketId` | `string`<br>Name of the output bucket
-`condition` | `string`<br>Filter for agents to include
-`acl` | [AclConfig](#aclconfig)<br>Access control list
-`uba` | [UbaComputeClusterConfig](#ubacomputeclusterconfig)<br>Configuration settings for Unreal Build Accelerator (UBA) for a compute cluster
+`id` | `string`<br>
+`namespaceId` | `string`<br>
+`requestBucketId` | `string`<br>
+`responseBucketId` | `string`<br>
+`condition` | `string`<br>
+`acl` | [AclConfig](#aclconfig)<br>
+`uba` | [UbaComputeClusterConfig](#ubacomputeclusterconfig)<br>
 
 ## UbaComputeClusterConfig
 
-Configuration settings for Unreal Build Accelerator (UBA) for a compute cluster Contains cache server connectivity and other UBA-related settings. Ensure all cache related fields are non-null to activate use of cache.
-
 Name | Description
 ---- | -----------
-`cacheEndpoint` | `string`<br>Host and port for connecting to the UBA cache server (host:port format) As specified with `-port=[PORT]` to UBA cache process. Serves all cache requests and must be accessible to both agents and initiators running UBA / Unreal Build Tool.
-`cacheHttpEndpoint` | `string`<br>Host and port for HTTP management of UBA cache server (host:port format) As specified with `-http=[PORT]` to UBA cache process. This port is used by Horde to manage the UBA cache, such as adding new cache session keys for initiators For improved security, only allow Horde server to access this port
-`cacheHttpCrypto` | `string`<br>AES encryption key used for secure Horde-to-cache-server communication As specified with `-httpcrypto=[KEY]` to UBA cache process.
+`cacheEndpoint` | `string`<br>
+`cacheHttpEndpoint` | `string`<br>
+`cacheHttpCrypto` | `string`<br>
 
 ## PoolConfig
 
-Mutable configuration for a pool
-
 Name | Description
 ---- | -----------
-`id` | `string`<br>Unique id for this pool
-`base` | `string`<br>Base pool config to copy settings from
-`name` | `string`<br>Name of the pool
-`condition` | `string`<br>Condition for agents to automatically be included in this pool
-`properties` | `string` `->` `string`<br>Arbitrary properties related to this pool
-`color` | [PoolColor](#poolcolor-enum)<br>Color to use for this pool on the dashboard
-`enableAutoscaling` | `boolean`<br>Whether to enable autoscaling for this pool
-`minAgents` | `integer`<br>The minimum number of agents to keep in the pool
-`numReserveAgents` | `integer`<br>The minimum number of idle agents to hold in reserve
-`conformInterval` | `string`<br>Interval between conforms. If zero, the pool will not conform on a schedule.
-`scaleOutCooldown` | `string`<br>Cooldown time between scale-out events
-`scaleInCooldown` | `string`<br>Cooldown time between scale-in events
-`shutdownIfDisabledGracePeriod` | `string`<br>Time to wait before shutting down an agent that has been disabled
+`id` | `string`<br>
+`base` | `string`<br>
+`name` | `string`<br>
+`condition` | `string`<br>
+`properties` | `string` `->` `string`<br>
+`color` | [PoolColor](#poolcolor-enum)<br>
+`enableAutoscaling` | `boolean`<br>
+`minAgents` | `integer`<br>
+`numReserveAgents` | `integer`<br>
+`conformInterval` | `string`<br>
+`scaleOutCooldown` | `string`<br>
+`scaleInCooldown` | `string`<br>
+`shutdownIfDisabledGracePeriod` | `string`<br>
 `sizeStrategy` | [PoolSizeStrategy](#poolsizestrategy-enum)<br>
-`sizeStrategies` | [PoolSizeStrategyInfo](#poolsizestrategyinfo)`[]`<br>List of pool sizing strategies for this pool. The first strategy with a matching condition will be picked.
-`fleetManagers` | [FleetManagerInfo](#fleetmanagerinfo)`[]`<br>List of fleet managers for this pool. The first strategy with a matching condition will be picked. If empty or no conditions match, a default fleet manager will be used.
-`leaseUtilizationSettings` | [LeaseUtilizationSettings](#leaseutilizationsettings)<br>Settings for lease utilization pool sizing strategy (if used)
-`jobQueueSettings` | [JobQueueSettings](#jobqueuesettings)<br>Settings for job queue pool sizing strategy (if used)
+`sizeStrategies` | [PoolSizeStrategyInfo](#poolsizestrategyinfo)`[]`<br>
+`fleetManagers` | [FleetManagerInfo](#fleetmanagerinfo)`[]`<br>
+`leaseUtilizationSettings` | [LeaseUtilizationSettings](#leaseutilizationsettings)<br>
+`jobQueueSettings` | [JobQueueSettings](#jobqueuesettings)<br>
 `computeQueueAwsMetricSettings` | [ComputeQueueAwsMetricSettings](#computequeueawsmetricsettings)<br>
 
 ## PoolColor (Enum)
 
-Color to use for labels of this pool
-
 Name | Description
 ---- | -----------
 `Default` | 
@@ -345,323 +307,259 @@ Name | Description
 
 ## PoolSizeStrategy (Enum)
 
-Available pool sizing strategies
-
 Name | Description
 ---- | -----------
-`LeaseUtilization` | Strategy based on lease utilization
-`JobQueue` | Strategy based on size of job build queue
-`NoOp` | No-op strategy used as fallback/default behavior
-`ComputeQueueAwsMetric` | A no-op strategy that reports metrics to let an external AWS auto-scaling policy scale the fleet
-`LeaseUtilizationAwsMetric` | A no-op strategy that reports metrics to let an external AWS auto-scaling policy scale the fleet
+`LeaseUtilization` | 
+`JobQueue` | 
+`NoOp` | 
+`ComputeQueueAwsMetric` | 
+`LeaseUtilizationAwsMetric` | 
 
 ## PoolSizeStrategyInfo
 
-Metadata for configuring and picking a pool sizing strategy
-
 Name | Description
 ---- | -----------
-`type` | [PoolSizeStrategy](#poolsizestrategy-enum)<br>Strategy implementation to use
-`condition` | `string`<br>Condition if this strategy should be enabled (right now, using date/time as a distinguishing factor)
-`config` | `object`<br>Configuration for the strategy, serialized as JSON
-`extraAgentCount` | `integer`<br>Integer to add after pool size has been calculated. Can also be negative.
+`type` | [PoolSizeStrategy](#poolsizestrategy-enum)<br>
+`condition` | `string`<br>
+`config` | `object`<br>
+`extraAgentCount` | `integer`<br>
 
 ## FleetManagerInfo
 
-Metadata for configuring and picking a fleet manager
-
 Name | Description
 ---- | -----------
-`type` | [FleetManagerType](#fleetmanagertype-enum)<br>Fleet manager type implementation to use
-`condition` | `string`<br>Condition if this strategy should be enabled (right now, using date/time as a distinguishing factor)
-`config` | `object`<br>Configuration for the strategy, serialized as JSON
+`type` | [FleetManagerType](#fleetmanagertype-enum)<br>
+`condition` | `string`<br>
+`config` | `object`<br>
 
 ## FleetManagerType (Enum)
 
-Available fleet managers
-
 Name | Description
 ---- | -----------
-`Default` | Default fleet manager
-`NoOp` | No-op fleet manager.
-`Aws` | Fleet manager for handling AWS EC2 instances. Will create and/or terminate instances from scratch.
-`AwsReuse` | Fleet manager for handling AWS EC2 instances. Will start already existing but stopped instances to reuse existing EBS disks.
-`AwsRecycle` | Fleet manager for handling AWS EC2 instances. Will start already existing but stopped instances to reuse existing EBS disks.
-`AwsAsg` | Fleet manager for handling AWS EC2 instances. Uses an EC2 auto-scaling group for controlling the number of running instances.
+`Default` | 
+`NoOp` | 
+`Aws` | 
+`AwsReuse` | 
+`AwsRecycle` | 
+`AwsAsg` | 
 
 ## LeaseUtilizationSettings
 
-Lease utilization sizing settings for a pool
-
 Name | Description
 ---- | -----------
-`sampleTimeSec` | `integer`<br>Time period for each sample
-`numSamples` | `integer`<br>Number of samples to collect for calculating lease utilization
-`numSamplesForResult` | `integer`<br>Min number of samples for a valid result
-`minAgents` | `integer`<br>The minimum number of agents to keep in the pool
-`numReserveAgents` | `integer`<br>The minimum number of idle agents to hold in reserve
+`sampleTimeSec` | `integer`<br>
+`numSamples` | `integer`<br>
+`numSamplesForResult` | `integer`<br>
+`minAgents` | `integer`<br>
+`numReserveAgents` | `integer`<br>
 
 ## JobQueueSettings
 
-Job queue sizing settings for a pool
-
 Name | Description
 ---- | -----------
-`scaleOutFactor` | `number`<br>Factor translating queue size to additional agents to grow the pool with The result is always rounded up to nearest integer. Example: if there are 20 jobs in queue, a factor 0.25 will result in 5 new agents being added (20 * 0.25)
-`scaleInFactor` | `number`<br>Factor by which to shrink the pool size with when queue is empty The result is always rounded up to nearest integer. Example: when the queue size is zero, a default value of 0.9 will shrink the pool by 10% (current agent count * 0.9)
-`samplePeriodMin` | `integer`<br>How far back in time to look for job batches (that potentially are in the queue)
-`readyTimeThresholdSec` | `integer`<br>Time spent in ready state before considered truly waiting for an agent<br>A job batch can be in ready state before getting picked up and executed. This threshold will help ensure only batches that have been waiting longer than this value will be considered.
+`scaleOutFactor` | `number`<br>
+`scaleInFactor` | `number`<br>
+`samplePeriodMin` | `integer`<br>
+`readyTimeThresholdSec` | `integer`<br>
 
 ## ComputeQueueAwsMetricSettings
 
-Settings for
-
 Name | Description
 ---- | -----------
-`computeClusterId` | `string`<br>Compute cluster ID to observe
-`namespace` | `string`<br>AWS CloudWatch namespace to write metrics in
+`computeClusterId` | `string`<br>
+`namespace` | `string`<br>
 
 ## AgentSoftwareConfig
 
-Selects different agent software versions by evaluating a condition
-
 Name | Description
 ---- | -----------
-`toolId` | `string`<br>Tool identifier
-`condition` | `string`<br>Condition for using this channel
+`toolId` | `string`<br>
+`condition` | `string`<br>
 
 ## NetworkConfig
 
-Describes a network The ID describes any logical grouping, such as region, availability zone, rack or office location.
-
-Name | Description
----- | -----------
-`id` | `string`<br>ID for this network
-`cidrBlock` | `string`<br>CIDR block
-`description` | `string`<br>Human-readable description
-`computeId` | `string`<br>Compute ID for this network (used when allocating compute resources)
-
-## AnalyticsConfig
-
-Config settings for analytics
-
 Name | Description
 ---- | -----------
-`stores` | [TelemetryStoreConfig](Telemetry.md)`[]`<br>Metrics to aggregate on the Horde server
-
-## ExamplePluginConfig
+`id` | `string`<br>
+`cidrBlock` | `string`<br>
+`description` | `string`<br>
+`computeId` | `string`<br>
 
-Configuration for the ExamplePlugin plugin
+## ExperimentalConfig
 
 Name | Description
 ---- | -----------
-`greetingMessage` | `string`<br>The Greeting provided by the example plugin
+`notifications` | [NotificationConfig](Notification.md)`[]`<br>
 
 ## SecretsConfig
 
-Configuration for the secrets system
-
 Name | Description
 ---- | -----------
-`secrets` | [SecretConfig](#secretconfig)`[]`<br>List of secrets
-`providerConfigs` | [SecretProviderConfig](#secretproviderconfig)`[]`<br>List of configurations for secret providers
+`secrets` | [SecretConfig](#secretconfig)`[]`<br>
+`providerConfigs` | [SecretProviderConfig](#secretproviderconfig)`[]`<br>
 
 ## SecretConfig
 
-Configuration for a secret value
-
 Name | Description
 ---- | -----------
-`id` | `string`<br>Identifier for this secret
-`data` | `string` `->` `string`<br>Key/value pairs associated with this secret
-`sources` | [ExternalSecretConfig](#externalsecretconfig)`[]`<br>Providers to source key/value pairs from
-`acl` | [AclConfig](#aclconfig)<br>Defines access to this particular secret
+`id` | `string`<br>
+`data` | `string` `->` `string`<br>
+`sources` | [ExternalSecretConfig](#externalsecretconfig)`[]`<br>
+`acl` | [AclConfig](#aclconfig)<br>
 
 ## ExternalSecretConfig
 
-Configuration for an external secret provider
-
 Name | Description
 ---- | -----------
-`providerConfig` | `string`<br>The name to look up a secret provider config. A config is mutually exclusive to and to be used when a secret requires additional options than just the name of the provider
-`provider` | `string`<br>Name of the provider to use
-`format` | [ExternalSecretFormat](#externalsecretformat-enum)<br>Format of the secret
-`key` | `string`<br>Optional key indicating the parameter to set in the resulting data array. Required if  is .
-`path` | `string`<br>Optional value indicating what to fetch from the provider
+`providerConfig` | `string`<br>
+`provider` | `string`<br>
+`format` | [ExternalSecretFormat](#externalsecretformat-enum)<br>
+`key` | `string`<br>
+`path` | `string`<br>
 
 ## ExternalSecretFormat (Enum)
 
-Format describing how to parse external secret values
-
 Name | Description
 ---- | -----------
-`Text` | Secret is a plain text value which will be stored using the external secret key
-`Json` | Secret is a JSON formatted string containing key/value pairs
+`Text` | 
+`Json` | 
 
 ## SecretProviderConfig
 
-Common settings object for different secret providers
-
 Name | Description
 ---- | -----------
-`name` | `string`<br>Name of the secret provider config
-`provider` | `string`<br>Name of the provider
-`hcpVault` | [HcpVaultConfig](#hcpvaultconfig)<br>Configuration for HCP Vault
-`azureKeyVault` | [AzureKeyVaultConfig](#azurekeyvaultconfig)<br>Configuration for Azure Key Vault
+`name` | `string`<br>
+`provider` | `string`<br>
+`hcpVault` | [HcpVaultConfig](#hcpvaultconfig)<br>
+`azureKeyVault` | [AzureKeyVaultConfig](#azurekeyvaultconfig)<br>
 
 ## HcpVaultConfig
 
-Configuration for HCP Vault secrets.
-
 Name | Description
 ---- | -----------
-`preSharedKey` | `string`<br>The Vault authentication token to use when the type of credentials is .<br>More information on this token can be found in the API documentation for X-Vault-Token
-`endPoint` | `string`<br>The address of the Vault server.
-`awsArnRole` | `string`<br>The optional AWS ARN to assume a role when logging into Vault. To be used when the type of credentials is . When not set the AWS SDK default credential search will be used.
-`awsIamServerId` | `string`<br>The value for the X-Vault-AWS-IAM-Server-ID header when the Vault server is configured with iam_server_id_header_value. Most likely to be the Vault server's DNS name. To be used when the type of credentials is .
-`role` | `string`<br>The name of the role in the Vault configuration for the login. This role name can be different to the name of the AWS IAM principal. To be used when the type of credentials is . When not set the name of IAM principal will be used.
-`credentials` | [HcpVaultCredentialsType](#hcpvaultcredentialstype-enum)<br>Type of credentials to use.
+`preSharedKey` | `string`<br>
+`endPoint` | `string`<br>
+`awsArnRole` | `string`<br>
+`awsIamServerId` | `string`<br>
+`role` | `string`<br>
+`credentials` | [HcpVaultCredentialsType](#hcpvaultcredentialstype-enum)<br>
 
 ## HcpVaultCredentialsType (Enum)
 
-Credentials to use to log into HCP Vault.
-
 Name | Description
 ---- | -----------
-`PreSharedKey` | Use a pre-shared token provided by Vault for authentication.
-`AwsAuth` | Use AWS IAM credentials to authenticate against the Vault server.
+`PreSharedKey` | 
+`AwsAuth` | 
 
 ## AzureKeyVaultConfig
 
-Configuration for Azure Key Vault secrets.
-
 Name | Description
 ---- | -----------
-`vaultUri` | `string`<br>The Uri to the Azure Key Vault.
-
-## ExperimentalConfig
-
-Configuration for the experimental plugin
-
-Name | Description
----- | -----------
-`notifications` | [NotificationConfig](Notification.md)`[]`<br>Notifications to be sent on the Horde server
+`vaultUri` | `string`<br>
 
 ## StorageConfig
 
-Configuration for storage
-
 Name | Description
 ---- | -----------
-`enableGc` | `boolean`<br>Whether to enable garbage collection
-`enableGcVerification` | `boolean`<br>Whether to enable garbage collection in verification mode (nothing deleted, just logging on access to deleted blobs)
-`backends` | [BackendConfig](#backendconfig)`[]`<br>List of storage backends
-`namespaces` | [NamespaceConfig](#namespaceconfig)`[]`<br>List of namespaces for storage
+`enableGc` | `boolean`<br>
+`enableGcVerification` | `boolean`<br>
+`backends` | [BackendConfig](#backendconfig)`[]`<br>
+`namespaces` | [NamespaceConfig](#namespaceconfig)`[]`<br>
 
 ## BackendConfig
 
-Common settings object for different providers
-
 Name | Description
 ---- | -----------
-`id` | `string`<br>The storage backend ID
-`base` | `string`<br>Base backend to copy default settings from
-`secondary` | `string`<br>Specifies another backend to read from if an object is not found in this one. Can be used when migrating data from one backend to another.
+`id` | `string`<br>
+`base` | `string`<br>
+`secondary` | `string`<br>
 `type` | [StorageBackendType](#storagebackendtype-enum)<br>
 `baseDir` | `string`<br>
-`awsBucketName` | `string`<br>Name of the bucket to use
-`awsBucketPath` | `string`<br>Base path within the bucket
-`awsCredentials` | [AwsCredentialsType](#awscredentialstype-enum)<br>Type of credentials to use
-`awsRole` | `string`<br>ARN of a role to assume
-`awsProfile` | `string`<br>The AWS profile to read credentials form
-`awsRegion` | `string`<br>Region to connect to
-`azureConnectionString` | `string`<br>Connection string for Azure
-`azureContainerName` | `string`<br>Name of the container
+`awsBucketName` | `string`<br>
+`awsBucketPath` | `string`<br>
+`awsCredentials` | [AwsCredentialsType](#awscredentialstype-enum)<br>
+`awsRole` | `string`<br>
+`awsProfile` | `string`<br>
+`awsRegion` | `string`<br>
+`azureConnectionString` | `string`<br>
+`azureContainerName` | `string`<br>
 `relayServer` | `string`<br>
 `relayToken` | `string`<br>
-`gcsBucketName` | `string`<br>Name of the GCS bucket to use
-`gcsBucketPath` | `string`<br>Base path within the bucket
+`gcsBucketName` | `string`<br>
+`gcsBucketPath` | `string`<br>
 
 ## StorageBackendType (Enum)
 
-Types of storage backend to use
-
 Name | Description
 ---- | -----------
-`FileSystem` | Local filesystem
-`Aws` | AWS S3
-`Azure` | Azure blob store
-`Gcs` | Google Cloud Storage
-`Memory` | In-memory only (for testing)
+`FileSystem` | 
+`Aws` | 
+`Azure` | 
+`Gcs` | 
+`Memory` | 
 
 ## AwsCredentialsType (Enum)
 
-Credentials to use for AWS
-
 Name | Description
 ---- | -----------
-`Default` | Use default credentials from the AWS SDK
-`Profile` | Read credentials from the  profile in the AWS config file
-`AssumeRole` | Assume a particular role. Should specify ARN in
-`AssumeRoleWebIdentity` | Assume a particular role using the current environment variables.
+`Default` | 
+`Profile` | 
+`AssumeRole` | 
+`AssumeRoleWebIdentity` | 
 
 ## NamespaceConfig
 
-Configuration of a particular namespace
-
 Name | Description
 ---- | -----------
-`id` | `string`<br>Identifier for this namespace
-`backend` | `string`<br>Backend to use for this namespace
-`prefix` | `string`<br>Prefix for items within this namespace
-`gcFrequencyHrs` | `number`<br>How frequently to run garbage collection, in hours.
-`gcDelayHrs` | `number`<br>How long to keep newly uploaded orphanned objects before allowing them to be deleted, in hours.
-`enableAliases` | `boolean`<br>Support querying exports by their aliases
-`acl` | [AclConfig](#aclconfig)<br>Access list for this namespace
-
-## SymbolsConfig
+`id` | `string`<br>
+`backend` | `string`<br>
+`prefix` | `string`<br>
+`gcFrequencyHrs` | `number`<br>
+`gcDelayHrs` | `number`<br>
+`enableAliases` | `boolean`<br>
+`acl` | [AclConfig](#aclconfig)<br>
 
-Configuration for the tools system
+## ExamplePluginConfig
 
 Name | Description
 ---- | -----------
-`stores` | [SymbolStoreConfig](#symbolstoreconfig)`[]`<br>List of symbol stores
-
-## SymbolStoreConfig
+`greetingMessage` | `string`<br>
 
-Configuration for a symbol store
+## ToolsConfig
 
 Name | Description
 ---- | -----------
-`id` | `string`<br>Identifier for this store
-`namespaceId` | `string`<br>Configuration for the symbol store backend
-`public` | `boolean`<br>Whether to make this store available without auth
-`acl` | [AclConfig](#aclconfig)<br>Access to the symbol store
-
-## ToolsConfig
+`tools` | [ToolConfig](#toolconfig)`[]`<br>
 
-Configuration for the tools system
+## ToolConfig
 
 Name | Description
 ---- | -----------
-`tools` | [ToolConfig](#toolconfig)`[]`<br>Tool configurations
+`id` | `string`<br>
+`name` | `string`<br>
+`description` | `string`<br>
+`category` | `string`<br>
+`group` | `string`<br>
+`platforms` | `string[]`<br>
+`public` | `boolean`<br>
+`showInUgs` | `boolean`<br>
+`showInDashboard` | `boolean`<br>
+`showInToolbox` | `boolean`<br>
+`metadata` | `string` `->` `string`<br>
+`namespaceId` | `string`<br>
+`acl` | [AclConfig](#aclconfig)<br>
 
-## ToolConfig
+## SymbolsConfig
 
-Options for configuring a tool
+Name | Description
+---- | -----------
+`stores` | [SymbolStoreConfig](#symbolstoreconfig)`[]`<br>
+
+## SymbolStoreConfig
 
 Name | Description
 ---- | -----------
-`id` | `string`<br>Unique identifier for the tool
-`name` | `string`<br>Name of the tool
-`description` | `string`<br>Description for the tool
-`category` | `string`<br>Category for the tool. Will cause the tool to be shown in a different tab in the dashboard.
-`group` | `string`<br>Grouping key for different variations of the same tool. The dashboard will show these together.
-`platforms` | `string[]`<br>Platforms for this tool. Takes the form of a NET RID (https://learn.microsoft.com/en-us/dotnet/core/rid-catalog).
-`public` | `boolean`<br>Whether this tool should be exposed for download on a public endpoint without authentication
-`showInUgs` | `boolean`<br>Whether to show this tool for download in the UGS tools menu
-`showInDashboard` | `boolean`<br>Whether to show this tool for download in the dashboard
-`showInToolbox` | `boolean`<br>Whether to show this tool for download in Unreal Toolbox
-`metadata` | `string` `->` `string`<br>Metadata for this tool
-`namespaceId` | `string`<br>Default namespace for new deployments of this tool
-`acl` | [AclConfig](#aclconfig)<br>Permissions for the tool
+`id` | `string`<br>
+`namespaceId` | `string`<br>
+`public` | `boolean`<br>
+`acl` | [AclConfig](#aclconfig)<br>
diff --git a/Engine/Source/Programs/Horde/Docs/Config/Schema/Issues.md b/Engine/Source/Programs/Horde/Docs/Config/Schema/Issues.md
--- a/Engine/Source/Programs/Horde/Docs/Config/Schema/Issues.md
+++ b/Engine/Source/Programs/Horde/Docs/Config/Schema/Issues.md
@@ -2,25 +2,19 @@
 
 # *.issues.json
 
-Stores configuration for build health issues
-
 Name | Description
 ---- | -----------
-`include` | [ConfigInclude](#configinclude)`[]`<br>Includes for other configuration files
-`channelGroups` | [IssueChannelGroupConfig](#issuechannelgroupconfig)`[]`<br>Issue grouping channels
+`include` | [ConfigInclude](#configinclude)`[]`<br>
+`channelGroups` | [IssueChannelGroupConfig](#issuechannelgroupconfig)`[]`<br>
 
 ## ConfigInclude
 
-Directive to merge config data from another source
-
 Name | Description
 ---- | -----------
-`path` | `string`<br>Path to the config data to be included. May be relative to the including file's location.
+`path` | `string`<br>
 
 ## IssueChannelGroupConfig
 
-Configuration for grouping to apply based on logging channels
-
 Name | Description
 ---- | -----------
-`channels` | `string[]`<br>A list of channel logging boundaries to group issues on
+`channels` | `string[]`<br>
diff --git a/Engine/Source/Programs/Horde/Docs/Config/Schema/Notification.md b/Engine/Source/Programs/Horde/Docs/Config/Schema/Notification.md
--- a/Engine/Source/Programs/Horde/Docs/Config/Schema/Notification.md
+++ b/Engine/Source/Programs/Horde/Docs/Config/Schema/Notification.md
@@ -2,50 +2,40 @@
 
 # *.notification.json
 
-Config for notifications
-
 Name | Description
 ---- | -----------
-`id` | `string`<br>Identifier for this store
-`notificationStreams` | [NotificationStreamConfig](#notificationstreamconfig)`[]`<br>Streams for the notifications to be sent on the Horde server
-`include` | [ConfigInclude](#configinclude)`[]`<br>Includes for other configuration files
+`id` | `string`<br>
+`notificationStreams` | [NotificationStreamConfig](#notificationstreamconfig)`[]`<br>
+`include` | [ConfigInclude](#configinclude)`[]`<br>
 
 ## NotificationStreamConfig
 
-Config for the notifications based on the streams
-
 Name | Description
 ---- | -----------
-`streams` | `string[]`<br>List of streams for this notification
-`enabledStreamTags` | `string[]`<br>List of enabled tags for a stream for this notification
-`jobNotifications` | [JobNotificationConfig](#jobnotificationconfig)`[]`<br>Configuration for sending experimental Job Summary Slack notifications
+`streams` | `string[]`<br>
+`enabledStreamTags` | `string[]`<br>
+`jobNotifications` | [JobNotificationConfig](#jobnotificationconfig)`[]`<br>
 
 ## JobNotificationConfig
 
-Configuration for the job notification
-
 Name | Description
 ---- | -----------
-`enabled` | `boolean`<br>Whether or not the current configuration is enabled
-`template` | `string`<br>Identifier of the template this notification applies to
-`namePattern` | `string`<br>Regex pattern used to extract name of the step this notification will use for display
-`notificationFormats` | [NotificationFormatConfig](#notificationformatconfig)`[]`<br>List of formats to be applied for this notification
-`channels` | `string[]`<br>List of Slack channel or user identifiers the notification will be sent to
+`enabled` | `boolean`<br>
+`template` | `string`<br>
+`namePattern` | `string`<br>
+`notificationFormats` | [NotificationFormatConfig](#notificationformatconfig)`[]`<br>
+`channels` | `string[]`<br>
 
 ## NotificationFormatConfig
 
-Configuration for the format of the notification group
-
 Name | Description
 ---- | -----------
-`group` | `string`<br>Name of the group
-`stepPattern` | `string`<br>Regex pattern to determine if a job step belongs within the group
-`alternativeNamePattern` | `string`<br>Optional regex pattern to override the
+`group` | `string`<br>
+`stepPattern` | `string`<br>
+`alternativeNamePattern` | `string`<br>
 
 ## ConfigInclude
 
-Directive to merge config data from another source
-
 Name | Description
 ---- | -----------
-`path` | `string`<br>Path to the config data to be included. May be relative to the including file's location.
+`path` | `string`<br>
diff --git a/Engine/Source/Programs/Horde/Docs/Config/Schema/Projects.md b/Engine/Source/Programs/Horde/Docs/Config/Schema/Projects.md
--- a/Engine/Source/Programs/Horde/Docs/Config/Schema/Projects.md
+++ b/Engine/Source/Programs/Horde/Docs/Config/Schema/Projects.md
@@ -2,75 +2,65 @@
 
 # *.project.json
 
-Stores configuration for a project
-
 Name | Description
 ---- | -----------
-`id` | `string`<br>The project id
-`name` | `string`<br>Name for the new project
-`path` | `string`<br>Direct include path for the project config. For backwards compatibility with old config files when including from a GlobalConfig object.
-`include` | [ConfigInclude](#configinclude)`[]`<br>Includes for other configuration files
-`macros` | [ConfigMacro](#configmacro)`[]`<br>Macros within the global scope
-`order` | `integer`<br>Order of this project on the dashboard
-`logo` | `string`<br>Path to the project logo
-`logoDarkTheme` | `string`<br>Optional path to the project logo for the dark theme
-`pools` | [PoolConfig](#poolconfig)`[]`<br>List of pools for this project
-`categories` | [ProjectCategoryConfig](#projectcategoryconfig)`[]`<br>Categories to include in this project
-`jobOptions` | [JobOptions](#joboptions)<br>Default settings for executing jobs
-`workspaceTypes` | `string` `->` [WorkspaceConfig](#workspaceconfig)<br>Default workspace types for streams These are added to the list of each stream's workspace types.
-`ignorePatternFiles` | `string[]`<br>Additional ignore pattern files applied to all streams in this project. Streams can append additional files via their own IgnorePatternFiles property. Paths are relative to workspace root.
-`telemetryStoreId` | `string`<br>Telemetry store for Horde data for this project
-`streams` | [StreamConfig](Streams.md)`[]`<br>List of streams
-`artifactTypes` | [ArtifactTypeConfig](#artifacttypeconfig)`[]`<br>Permissions for artifact types
-`acl` | [AclConfig](#aclconfig)<br>Acl entries
+`id` | `string`<br>
+`name` | `string`<br>
+`path` | `string`<br>
+`include` | [ConfigInclude](#configinclude)`[]`<br>
+`macros` | [ConfigMacro](#configmacro)`[]`<br>
+`order` | `integer`<br>
+`logo` | `string`<br>
+`logoDarkTheme` | `string`<br>
+`pools` | [PoolConfig](#poolconfig)`[]`<br>
+`categories` | [ProjectCategoryConfig](#projectcategoryconfig)`[]`<br>
+`jobOptions` | [JobOptions](#joboptions)<br>
+`workspaceTypes` | `string` `->` [WorkspaceConfig](#workspaceconfig)<br>
+`ignorePatternFiles` | `string[]`<br>
+`telemetryStoreId` | `string`<br>
+`streams` | [StreamConfig](Streams.md)`[]`<br>
+`artifactTypes` | [ArtifactTypeConfig](#artifacttypeconfig)`[]`<br>
+`acl` | [AclConfig](#aclconfig)<br>
 
 ## ConfigInclude
 
-Directive to merge config data from another source
-
 Name | Description
 ---- | -----------
-`path` | `string`<br>Path to the config data to be included. May be relative to the including file's location.
+`path` | `string`<br>
 
 ## ConfigMacro
 
-Declares a config macro
-
 Name | Description
 ---- | -----------
-`name` | `string`<br>Name of the macro property
-`value` | `string`<br>Value for the macro property
+`name` | `string`<br>
+`value` | `string`<br>
 
 ## PoolConfig
 
-Mutable configuration for a pool
-
 Name | Description
 ---- | -----------
-`id` | `string`<br>Unique id for this pool
-`base` | `string`<br>Base pool config to copy settings from
-`name` | `string`<br>Name of the pool
-`condition` | `string`<br>Condition for agents to automatically be included in this pool
-`properties` | `string` `->` `string`<br>Arbitrary properties related to this pool
-`color` | [PoolColor](#poolcolor-enum)<br>Color to use for this pool on the dashboard
-`enableAutoscaling` | `boolean`<br>Whether to enable autoscaling for this pool
-`minAgents` | `integer`<br>The minimum number of agents to keep in the pool
-`numReserveAgents` | `integer`<br>The minimum number of idle agents to hold in reserve
-`conformInterval` | `string`<br>Interval between conforms. If zero, the pool will not conform on a schedule.
-`scaleOutCooldown` | `string`<br>Cooldown time between scale-out events
-`scaleInCooldown` | `string`<br>Cooldown time between scale-in events
-`shutdownIfDisabledGracePeriod` | `string`<br>Time to wait before shutting down an agent that has been disabled
+`id` | `string`<br>
+`base` | `string`<br>
+`name` | `string`<br>
+`condition` | `string`<br>
+`properties` | `string` `->` `string`<br>
+`color` | [PoolColor](#poolcolor-enum)<br>
+`enableAutoscaling` | `boolean`<br>
+`minAgents` | `integer`<br>
+`numReserveAgents` | `integer`<br>
+`conformInterval` | `string`<br>
+`scaleOutCooldown` | `string`<br>
+`scaleInCooldown` | `string`<br>
+`shutdownIfDisabledGracePeriod` | `string`<br>
 `sizeStrategy` | [PoolSizeStrategy](#poolsizestrategy-enum)<br>
-`sizeStrategies` | [PoolSizeStrategyInfo](#poolsizestrategyinfo)`[]`<br>List of pool sizing strategies for this pool. The first strategy with a matching condition will be picked.
-`fleetManagers` | [FleetManagerInfo](#fleetmanagerinfo)`[]`<br>List of fleet managers for this pool. The first strategy with a matching condition will be picked. If empty or no conditions match, a default fleet manager will be used.
-`leaseUtilizationSettings` | [LeaseUtilizationSettings](#leaseutilizationsettings)<br>Settings for lease utilization pool sizing strategy (if used)
-`jobQueueSettings` | [JobQueueSettings](#jobqueuesettings)<br>Settings for job queue pool sizing strategy (if used)
+`sizeStrategies` | [PoolSizeStrategyInfo](#poolsizestrategyinfo)`[]`<br>
+`fleetManagers` | [FleetManagerInfo](#fleetmanagerinfo)`[]`<br>
+`leaseUtilizationSettings` | [LeaseUtilizationSettings](#leaseutilizationsettings)<br>
+`jobQueueSettings` | [JobQueueSettings](#jobqueuesettings)<br>
 `computeQueueAwsMetricSettings` | [ComputeQueueAwsMetricSettings](#computequeueawsmetricsettings)<br>
 
 ## PoolColor (Enum)
 
-Color to use for labels of this pool
-
 Name | Description
 ---- | -----------
 `Default` | 
@@ -81,93 +71,77 @@ Name | Description
 
 ## PoolSizeStrategy (Enum)
 
-Available pool sizing strategies
-
 Name | Description
 ---- | -----------
-`LeaseUtilization` | Strategy based on lease utilization
-`JobQueue` | Strategy based on size of job build queue
-`NoOp` | No-op strategy used as fallback/default behavior
-`ComputeQueueAwsMetric` | A no-op strategy that reports metrics to let an external AWS auto-scaling policy scale the fleet
-`LeaseUtilizationAwsMetric` | A no-op strategy that reports metrics to let an external AWS auto-scaling policy scale the fleet
+`LeaseUtilization` | 
+`JobQueue` | 
+`NoOp` | 
+`ComputeQueueAwsMetric` | 
+`LeaseUtilizationAwsMetric` | 
 
 ## PoolSizeStrategyInfo
 
-Metadata for configuring and picking a pool sizing strategy
-
 Name | Description
 ---- | -----------
-`type` | [PoolSizeStrategy](#poolsizestrategy-enum)<br>Strategy implementation to use
-`condition` | `string`<br>Condition if this strategy should be enabled (right now, using date/time as a distinguishing factor)
-`config` | `object`<br>Configuration for the strategy, serialized as JSON
-`extraAgentCount` | `integer`<br>Integer to add after pool size has been calculated. Can also be negative.
+`type` | [PoolSizeStrategy](#poolsizestrategy-enum)<br>
+`condition` | `string`<br>
+`config` | `object`<br>
+`extraAgentCount` | `integer`<br>
 
 ## FleetManagerInfo
 
-Metadata for configuring and picking a fleet manager
-
 Name | Description
 ---- | -----------
-`type` | [FleetManagerType](#fleetmanagertype-enum)<br>Fleet manager type implementation to use
-`condition` | `string`<br>Condition if this strategy should be enabled (right now, using date/time as a distinguishing factor)
-`config` | `object`<br>Configuration for the strategy, serialized as JSON
+`type` | [FleetManagerType](#fleetmanagertype-enum)<br>
+`condition` | `string`<br>
+`config` | `object`<br>
 
 ## FleetManagerType (Enum)
 
-Available fleet managers
-
 Name | Description
 ---- | -----------
-`Default` | Default fleet manager
-`NoOp` | No-op fleet manager.
-`Aws` | Fleet manager for handling AWS EC2 instances. Will create and/or terminate instances from scratch.
-`AwsReuse` | Fleet manager for handling AWS EC2 instances. Will start already existing but stopped instances to reuse existing EBS disks.
-`AwsRecycle` | Fleet manager for handling AWS EC2 instances. Will start already existing but stopped instances to reuse existing EBS disks.
-`AwsAsg` | Fleet manager for handling AWS EC2 instances. Uses an EC2 auto-scaling group for controlling the number of running instances.
+`Default` | 
+`NoOp` | 
+`Aws` | 
+`AwsReuse` | 
+`AwsRecycle` | 
+`AwsAsg` | 
 
 ## LeaseUtilizationSettings
 
-Lease utilization sizing settings for a pool
-
 Name | Description
 ---- | -----------
-`sampleTimeSec` | `integer`<br>Time period for each sample
-`numSamples` | `integer`<br>Number of samples to collect for calculating lease utilization
-`numSamplesForResult` | `integer`<br>Min number of samples for a valid result
-`minAgents` | `integer`<br>The minimum number of agents to keep in the pool
-`numReserveAgents` | `integer`<br>The minimum number of idle agents to hold in reserve
+`sampleTimeSec` | `integer`<br>
+`numSamples` | `integer`<br>
+`numSamplesForResult` | `integer`<br>
+`minAgents` | `integer`<br>
+`numReserveAgents` | `integer`<br>
 
 ## JobQueueSettings
 
-Job queue sizing settings for a pool
-
 Name | Description
 ---- | -----------
-`scaleOutFactor` | `number`<br>Factor translating queue size to additional agents to grow the pool with The result is always rounded up to nearest integer. Example: if there are 20 jobs in queue, a factor 0.25 will result in 5 new agents being added (20 * 0.25)
-`scaleInFactor` | `number`<br>Factor by which to shrink the pool size with when queue is empty The result is always rounded up to nearest integer. Example: when the queue size is zero, a default value of 0.9 will shrink the pool by 10% (current agent count * 0.9)
-`samplePeriodMin` | `integer`<br>How far back in time to look for job batches (that potentially are in the queue)
-`readyTimeThresholdSec` | `integer`<br>Time spent in ready state before considered truly waiting for an agent<br>A job batch can be in ready state before getting picked up and executed. This threshold will help ensure only batches that have been waiting longer than this value will be considered.
+`scaleOutFactor` | `number`<br>
+`scaleInFactor` | `number`<br>
+`samplePeriodMin` | `integer`<br>
+`readyTimeThresholdSec` | `integer`<br>
 
 ## ComputeQueueAwsMetricSettings
 
-Settings for
-
 Name | Description
 ---- | -----------
-`computeClusterId` | `string`<br>Compute cluster ID to observe
-`namespace` | `string`<br>AWS CloudWatch namespace to write metrics in
+`computeClusterId` | `string`<br>
+`namespace` | `string`<br>
 
 ## ProjectCategoryConfig
 
-Information about a category to display for a stream
-
 Name | Description
 ---- | -----------
-`name` | `string`<br>Name of this category
-`row` | `integer`<br>Index of the row to display this category on
-`showOnNavMenu` | `boolean`<br>Whether to show this category on the nav menu
-`includePatterns` | `string[]`<br>Patterns for stream names to include
-`excludePatterns` | `string[]`<br>Patterns for stream names to exclude
+`name` | `string`<br>
+`row` | `integer`<br>
+`showOnNavMenu` | `boolean`<br>
+`includePatterns` | `string[]`<br>
+`excludePatterns` | `string[]`<br>
 
 ## JobOptions
 
@@ -196,74 +170,62 @@ Name | Description
 
 ## WorkspaceConfig
 
-Information about a workspace type
-
 Name | Description
 ---- | -----------
 `base` | `string`<br>
-`cluster` | `string`<br>Name of the Perforce server cluster to use
-`serverAndPort` | `string`<br>The Perforce server and port (eg. perforce:1666)
-`userName` | `string`<br>User to log into Perforce with (defaults to buildmachine)
-`identifier` | `string`<br>Identifier to distinguish this workspace from other workspaces. Defaults to the workspace type name.
-`stream` | `string`<br>Override for the stream to sync
-`view` | `string[]`<br>Custom view for the workspace
-`incremental` | `boolean`<br>Whether to use an incrementally synced workspace
-`useAutoSdk` | `boolean`<br>Whether to use the AutoSDK
-`autoSdkView` | `string[]`<br>View for the AutoSDK paths to sync. If null, the whole thing will be synced.
-`method` | `string`<br>Method to use when syncing/materializing data from Perforce
-`minScratchSpace` | `integer`<br>Minimum disk space that must be available *after* syncing this workspace (in megabytes) If not available, the job will be aborted.
-`conformDiskFreeSpace` | `integer`<br>Threshold for when to trigger an automatic conform of agent. Measured in megabytes free on disk. Set to null or 0 to disable.
+`cluster` | `string`<br>
+`serverAndPort` | `string`<br>
+`userName` | `string`<br>
+`identifier` | `string`<br>
+`stream` | `string`<br>
+`view` | `string[]`<br>
+`incremental` | `boolean`<br>
+`useAutoSdk` | `boolean`<br>
+`autoSdkView` | `string[]`<br>
+`method` | `string`<br>
+`minScratchSpace` | `integer`<br>
+`conformDiskFreeSpace` | `integer`<br>
 
 ## ArtifactTypeConfig
 
-Configuration for an artifact
-
 Name | Description
 ---- | -----------
-`name` | `string`<br>Legacy 'Name' property
-`type` | `string`<br>Name of the artifact type
-`acl` | [AclConfig](#aclconfig)<br>Acl for the artifact type
-`keepCount` | `integer`<br>Number of artifacts to retain
-`keepDays` | `integer`<br>Number of days to retain artifacts of this type
-`namespaceId` | `string`<br>Storage namespace to use for this artifact types
+`name` | `string`<br>
+`type` | `string`<br>
+`acl` | [AclConfig](#aclconfig)<br>
+`keepCount` | `integer`<br>
+`keepDays` | `integer`<br>
+`namespaceId` | `string`<br>
 
 ## AclConfig
 
-Parameters to update an ACL
-
 Name | Description
 ---- | -----------
-`entries` | [AclEntryConfig](#aclentryconfig)`[]`<br>Entries to replace the existing ACL
-`profiles` | [AclProfileConfig](#aclprofileconfig)`[]`<br>Defines profiles which allow grouping sets of actions into named collections
-`inherit` | `boolean`<br>Whether to inherit permissions from the parent ACL
-`exceptions` | `string[]`<br>List of exceptions to the inherited setting
+`entries` | [AclEntryConfig](#aclentryconfig)`[]`<br>
+`profiles` | [AclProfileConfig](#aclprofileconfig)`[]`<br>
+`inherit` | `boolean`<br>
+`exceptions` | `string[]`<br>
 
 ## AclEntryConfig
 
-Individual entry in an ACL
-
 Name | Description
 ---- | -----------
-`claim` | [AclClaimConfig](#aclclaimconfig)<br>Name of the user or group
-`actions` | `string[]`<br>Array of actions to allow
-`profiles` | `string[]`<br>List of profiles to grant
+`claim` | [AclClaimConfig](#aclclaimconfig)<br>
+`actions` | `string[]`<br>
+`profiles` | `string[]`<br>
 
 ## AclClaimConfig
 
-New claim to create
-
 Name | Description
 ---- | -----------
-`type` | `string`<br>The claim type
-`value` | `string`<br>The claim value
+`type` | `string`<br>
+`value` | `string`<br>
 
 ## AclProfileConfig
 
-Configuration for an ACL profile. This defines a preset group of actions which can be given to a user via an ACL entry.
-
 Name | Description
 ---- | -----------
-`id` | `string`<br>Identifier for this profile
-`actions` | `string[]`<br>Actions to include
-`excludeActions` | `string[]`<br>Actions to exclude from the inherited actions
-`extends` | `string[]`<br>Other profiles to extend from
+`id` | `string`<br>
+`actions` | `string[]`<br>
+`excludeActions` | `string[]`<br>
+`extends` | `string[]`<br>
diff --git a/Engine/Source/Programs/Horde/Docs/Config/Schema/Streams.md b/Engine/Source/Programs/Horde/Docs/Config/Schema/Streams.md
--- a/Engine/Source/Programs/Horde/Docs/Config/Schema/Streams.md
+++ b/Engine/Source/Programs/Horde/Docs/Config/Schema/Streams.md
@@ -2,63 +2,57 @@
 
 # *.stream.json
 
-Config for a stream
-
-Name | Description
----- | -----------
-`id` | `string`<br>Identifier for the stream
-`vCS` | `string`<br>Version control service for stream commits
-`repositoryName` | `string`<br>Repository name for version control systems that support multiple repositories
-`defaultBranchName` | `string`<br>Default branch name for version control systems that use branches (e.g., "main"). If not set, defaults to "main".
-`path` | `string`<br>Direct include path for the stream config. For backwards compatibility with old config files when including from a ProjectConfig object.
-`include` | [ConfigInclude](#configinclude)`[]`<br>Includes for other configuration files
-`macros` | [ConfigMacro](#configmacro)`[]`<br>Macros within this stream
-`name` | `string`<br>Name of the stream
-`enginePath` | `string`<br>Path to the engine directory within the workspace. Used for launching UAT.
-`ignorePatternFiles` | `string[]`<br>Additional file paths (relative to workspace root) containing ignore patterns for the log parser. Each file should contain one regex per line, with lines starting with '#' treated as comments. Inherits from ProjectConfig and appends stream-level entries.
-`clusterName` | `string`<br>The perforce cluster containing the stream
-`order` | `integer`<br>Order for this stream
-`initialAgentType` | `string`<br>Default initial agent type for templates
-`notificationChannel` | `string`<br>Notification channel for all jobs in this stream
-`notificationChannelFilter` | `string`<br>Notification channel filter for this template. Can be Success, Failure, or Warnings.
-`triageChannel` | `string`<br>Channel to post issue triage notifications
-`jobOptions` | [JobOptions](#joboptions)<br>Default settings for executing jobs
-`telemetryStoreId` | `string`<br>Telemetry store for Horde data for this stream
-`autoSdkView` | `string[]`<br>View for the AutoSDK paths to sync. If null, the whole thing will be synced.
-`defaultPreflightTemplate` | `string`<br>Legacy name for the default preflight template
-`defaultPreflight` | [DefaultPreflightConfig](#defaultpreflightconfig)<br>Default template for running preflights
-`commitTags` | [CommitTagConfig](#committagconfig)`[]`<br>List of tags to apply to commits. Allows fast searching and classification of different commit types (eg. code vs content).
-`tabs` | [TabConfig](#tabconfig)`[]`<br>List of tabs to show for the new stream
-`environment` | `string` `->` `string`<br>Global environment variables for all agents in this stream
-`agentTypes` | `string` `->` [AgentConfig](#agentconfig)<br>Map of agent name to type
-`workspaceTypes` | `string` `->` [WorkspaceConfig](#workspaceconfig)<br>Map of workspace name to type
-`templates` | [TemplateRefConfig](#templaterefconfig)`[]`<br>List of templates to create
-`aliases` | [StreamAliasConfig](#streamaliasconfig)`[]`<br>Stream aliases configured for this stream. Each alias provides a named subset of the stream's templates.
-`acl` | [AclConfig](#aclconfig)<br>Custom permissions for this object
-`pausedUntil` | `string`<br>Pause stream builds until specified date
-`pauseComment` | `string`<br>Reason for pausing builds of the stream
-`replicators` | [ReplicatorConfig](#replicatorconfig)`[]`<br>Configuration for workers to replicate commit data into Horde Storage.
-`workflows` | [WorkflowConfig](#workflowconfig)`[]`<br>Workflows for dealing with new issues
-`tokens` | [TokenConfig](#tokenconfig)`[]`<br>Tokens to create for each job step
-`artifactTypes` | [ArtifactTypeConfig](#artifacttypeconfig)`[]`<br>Permissions for artifact types
-`streamTags` | [StreamTag](#streamtag)`[]`<br>Tags for the stream
+Name | Description
+---- | -----------
+`id` | `string`<br>
+`vCS` | `string`<br>
+`repositoryName` | `string`<br>
+`defaultBranchName` | `string`<br>
+`path` | `string`<br>
+`include` | [ConfigInclude](#configinclude)`[]`<br>
+`macros` | [ConfigMacro](#configmacro)`[]`<br>
+`name` | `string`<br>
+`enginePath` | `string`<br>
+`ignorePatternFiles` | `string[]`<br>
+`clusterName` | `string`<br>
+`order` | `integer`<br>
+`initialAgentType` | `string`<br>
+`notificationChannel` | `string`<br>
+`notificationChannelFilter` | `string`<br>
+`triageChannel` | `string`<br>
+`jobOptions` | [JobOptions](#joboptions)<br>
+`telemetryStoreId` | `string`<br>
+`autoSdkView` | `string[]`<br>
+`defaultPreflightTemplate` | `string`<br>
+`defaultPreflight` | [DefaultPreflightConfig](#defaultpreflightconfig)<br>
+`commitTags` | [CommitTagConfig](#committagconfig)`[]`<br>
+`tabs` | [TabConfig](#tabconfig)`[]`<br>
+`environment` | `string` `->` `string`<br>
+`agentTypes` | `string` `->` [AgentConfig](#agentconfig)<br>
+`workspaceTypes` | `string` `->` [WorkspaceConfig](#workspaceconfig)<br>
+`templates` | [TemplateRefConfig](#templaterefconfig)`[]`<br>
+`aliases` | [StreamAliasConfig](#streamaliasconfig)`[]`<br>
+`acl` | [AclConfig](#aclconfig)<br>
+`pausedUntil` | `string`<br>
+`pauseComment` | `string`<br>
+`replicators` | [ReplicatorConfig](#replicatorconfig)`[]`<br>
+`workflows` | [WorkflowConfig](#workflowconfig)`[]`<br>
+`tokens` | [TokenConfig](#tokenconfig)`[]`<br>
+`artifactTypes` | [ArtifactTypeConfig](#artifacttypeconfig)`[]`<br>
+`streamTags` | [StreamTag](#streamtag)`[]`<br>
 
 ## ConfigInclude
 
-Directive to merge config data from another source
-
 Name | Description
 ---- | -----------
-`path` | `string`<br>Path to the config data to be included. May be relative to the including file's location.
+`path` | `string`<br>
 
 ## ConfigMacro
 
-Declares a config macro
-
 Name | Description
 ---- | -----------
-`name` | `string`<br>Name of the macro property
-`value` | `string`<br>Value for the macro property
+`name` | `string`<br>
+`value` | `string`<br>
 
 ## JobOptions
 
@@ -87,25 +81,21 @@ Name | Description
 
 ## DefaultPreflightConfig
 
-Specifies defaults for running a preflight
-
 Name | Description
 ---- | -----------
-`templateId` | `string`<br>The template id to query
-`change` | [ChangeQueryConfig](#changequeryconfig)<br>Query for the change to use
+`templateId` | `string`<br>
+`change` | [ChangeQueryConfig](#changequeryconfig)<br>
 
 ## ChangeQueryConfig
 
-Query selecting the base changelist to use
-
 Name | Description
 ---- | -----------
-`name` | `string`<br>Name of this query, for display on the dashboard.
-`condition` | `string`<br>Condition to evaluate before deciding to use this query. May query tags in a preflight.
-`templateId` | `string`<br>The template id to query
-`target` | `string`<br>The target to query
-`outcomes` | [JobStepOutcome](#jobstepoutcome-enum)`[]`<br>Whether to match a job that produced warnings
-`commitTag` | `string`<br>Finds the last commit with this tag
+`name` | `string`<br>
+`condition` | `string`<br>
+`templateId` | `string`<br>
+`target` | `string`<br>
+`outcomes` | [JobStepOutcome](#jobstepoutcome-enum)`[]`<br>
+`commitTag` | `string`<br>
 
 ## JobStepOutcome (Enum)
 
@@ -120,28 +110,24 @@ Name | Description
 
 ## CommitTagConfig
 
-Configuration for custom commit filters
-
 Name | Description
 ---- | -----------
-`name` | `string`<br>Name of the tag
-`base` | `string`<br>Base tag to copy settings from
-`filter` | `string[]`<br>List of files to be included in this filter
+`name` | `string`<br>
+`base` | `string`<br>
+`filter` | `string[]`<br>
 
 ## TabConfig
 
-Information about a page to display in the dashboard for a stream
-
 Name | Description
 ---- | -----------
-`title` | `string`<br>Title of this page
-`type` | `string`<br>Type of this tab
-`style` | [TabStyle](#tabstyle-enum)<br>Presentation style for this page
-`showNames` | `boolean`<br>Whether to show job names on this page
-`showPreflights` | `boolean`<br>Whether to show all user preflights
-`jobNames` | `string[]`<br>Names of jobs to include on this page. If there is only one name specified, the name column does not need to be displayed.
-`templates` | `string[]`<br>List of job template names to show on this page.
-`columns` | [TabColumnConfig](#tabcolumnconfig)`[]`<br>Columns to display for different types of aggregates
+`title` | `string`<br>
+`type` | `string`<br>
+`style` | [TabStyle](#tabstyle-enum)<br>
+`showNames` | `boolean`<br>
+`showPreflights` | `boolean`<br>
+`jobNames` | `string[]`<br>
+`templates` | `string[]`<br>
+`columns` | [TabColumnConfig](#tabcolumnconfig)`[]`<br>
 
 ## TabStyle (Enum)
 
@@ -154,15 +140,13 @@ Name | Description
 
 ## TabColumnConfig
 
-Describes a column to display on the jobs page
-
 Name | Description
 ---- | -----------
-`type` | [TabColumnType](#tabcolumntype-enum)<br>The type of column
-`heading` | `string`<br>Heading for this column
-`category` | `string`<br>Category of aggregates to display in this column. If null, includes any aggregate not matched by another column.
-`parameter` | `string`<br>Parameter to show in this column
-`relativeWidth` | `integer`<br>Relative width of this column.
+`type` | [TabColumnType](#tabcolumntype-enum)<br>
+`heading` | `string`<br>
+`category` | `string`<br>
+`parameter` | `string`<br>
+`relativeWidth` | `integer`<br>
 
 ## TabColumnType (Enum)
 
@@ -175,136 +159,118 @@ Name | Description
 
 ## AgentConfig
 
-Mapping from a BuildGraph agent type to a set of machines on the farm
-
 Name | Description
 ---- | -----------
-`base` | `string`<br>Base agent config to inherit settings from
-`pool` | `string`<br>Pool of agents to use for this agent type
-`workspace` | `string`<br>Name of the workspace to sync
-`tempStorageDir` | `string`<br>Path to the temporary storage dir
-`environment` | `string` `->` `string`<br>Environment variables to be set when executing the job
-`tokens` | [TokenConfig](#tokenconfig)`[]`<br>Tokens to allocate for this agent type
+`base` | `string`<br>
+`pool` | `string`<br>
+`workspace` | `string`<br>
+`tempStorageDir` | `string`<br>
+`environment` | `string` `->` `string`<br>
+`tokens` | [TokenConfig](#tokenconfig)`[]`<br>
 
 ## TokenConfig
 
-Configuration for allocating access tokens for each job
-
 Name | Description
 ---- | -----------
-`url` | `string`<br>URL to request tokens from
-`clientId` | `string`<br>Client id to use to request a new token
-`clientSecret` | `string`<br>Client secret to request a new access token
-`envVar` | `string`<br>Environment variable to set with the access token
+`url` | `string`<br>
+`clientId` | `string`<br>
+`clientSecret` | `string`<br>
+`envVar` | `string`<br>
 
 ## WorkspaceConfig
 
-Information about a workspace type
-
 Name | Description
 ---- | -----------
 `base` | `string`<br>
-`cluster` | `string`<br>Name of the Perforce server cluster to use
-`serverAndPort` | `string`<br>The Perforce server and port (eg. perforce:1666)
-`userName` | `string`<br>User to log into Perforce with (defaults to buildmachine)
-`identifier` | `string`<br>Identifier to distinguish this workspace from other workspaces. Defaults to the workspace type name.
-`stream` | `string`<br>Override for the stream to sync
-`view` | `string[]`<br>Custom view for the workspace
-`incremental` | `boolean`<br>Whether to use an incrementally synced workspace
-`useAutoSdk` | `boolean`<br>Whether to use the AutoSDK
-`autoSdkView` | `string[]`<br>View for the AutoSDK paths to sync. If null, the whole thing will be synced.
-`method` | `string`<br>Method to use when syncing/materializing data from Perforce
-`minScratchSpace` | `integer`<br>Minimum disk space that must be available *after* syncing this workspace (in megabytes) If not available, the job will be aborted.
-`conformDiskFreeSpace` | `integer`<br>Threshold for when to trigger an automatic conform of agent. Measured in megabytes free on disk. Set to null or 0 to disable.
+`cluster` | `string`<br>
+`serverAndPort` | `string`<br>
+`userName` | `string`<br>
+`identifier` | `string`<br>
+`stream` | `string`<br>
+`view` | `string[]`<br>
+`incremental` | `boolean`<br>
+`useAutoSdk` | `boolean`<br>
+`autoSdkView` | `string[]`<br>
+`method` | `string`<br>
+`minScratchSpace` | `integer`<br>
+`conformDiskFreeSpace` | `integer`<br>
 
 ## TemplateRefConfig
 
-Parameters to create a template within a stream
-
-Name | Description
----- | -----------
-`id` | `string`<br>Optional identifier for this ref. If not specified, an id will be generated from the name.
-`base` | `string`<br>Base template id to copy from
-`showUgsBadges` | `boolean`<br>Whether to show badges in UGS for these jobs
-`showUgsAlerts` | `boolean`<br>Whether to show alerts in UGS for these jobs
-`notificationChannel` | `string`<br>Notification channel for this template. Overrides the stream channel if set.
-`notificationChannelFilter` | `string`<br>Notification channel filter for this template. Can be a combination of "Success", "Failure" and "Warnings" separated by pipe characters.
-`triageChannel` | `string`<br>Triage channel for this template. Overrides the stream channel if set.
-`workflowId` | `string`<br>Workflow to user for this stream
-`annotations` | `string` `->` `string`<br>Default annotations to apply to nodes in this template
-`schedule` | [ScheduleConfig](#scheduleconfig)<br>Schedule to execute this template
-`chainedJobs` | [ChainedJobTemplateConfig](#chainedjobtemplateconfig)`[]`<br>List of chained job triggers
-`acl` | [AclConfig](#aclconfig)<br>The ACL for this template
-`name` | `string`<br>Name for the new template
-`description` | `string`<br>Description for the template
-`priority` | [Priority](#priority-enum)<br>Default priority for this job
-`allowPreflights` | `boolean`<br>Whether to allow preflights of this template
-`updateIssues` | `boolean`<br>Whether issues should be updated for all jobs using this template
-`promoteIssuesByDefault` | `boolean`<br>Whether issues should be promoted by default for this template, promoted issues will generate user notifications
-`initialAgentType` | `string`<br>Initial agent type to parse the buildgraph script on
-`submitNewChange` | `string`<br>Path to a file within the stream to submit to generate a new changelist for jobs
-`submitDescription` | `string`<br>Description for new changelists
-`defaultChange` | [ChangeQueryConfig](#changequeryconfig)`[]`<br>Default change to build at. Each object has a condition parameter which can evaluated by the server to determine which change to use.
-`arguments` | `string[]`<br>Fixed arguments for the new job
-`parameters` | [TemplateTextParameterConfig](#templatetextparameterconfig)/[TemplateListParameterConfig](#templatelistparameterconfig)/[TemplateBoolParameterConfig](#templateboolparameterconfig)`[]`<br>Parameters for this template
-`jobOptions` | [JobOptions](#joboptions)<br>Default settings for jobs
+Name | Description
+---- | -----------
+`id` | `string`<br>
+`base` | `string`<br>
+`showUgsBadges` | `boolean`<br>
+`showUgsAlerts` | `boolean`<br>
+`notificationChannel` | `string`<br>
+`notificationChannelFilter` | `string`<br>
+`triageChannel` | `string`<br>
+`workflowId` | `string`<br>
+`annotations` | `string` `->` `string`<br>
+`schedule` | [ScheduleConfig](#scheduleconfig)<br>
+`chainedJobs` | [ChainedJobTemplateConfig](#chainedjobtemplateconfig)`[]`<br>
+`acl` | [AclConfig](#aclconfig)<br>
+`name` | `string`<br>
+`description` | `string`<br>
+`priority` | [Priority](#priority-enum)<br>
+`allowPreflights` | `boolean`<br>
+`updateIssues` | `boolean`<br>
+`promoteIssuesByDefault` | `boolean`<br>
+`initialAgentType` | `string`<br>
+`submitNewChange` | `string`<br>
+`submitDescription` | `string`<br>
+`defaultChange` | [ChangeQueryConfig](#changequeryconfig)`[]`<br>
+`arguments` | `string[]`<br>
+`parameters` | [TemplateTextParameterConfig](#templatetextparameterconfig)/[TemplateListParameterConfig](#templatelistparameterconfig)/[TemplateBoolParameterConfig](#templateboolparameterconfig)`[]`<br>
+`jobOptions` | [JobOptions](#joboptions)<br>
 
 ## ScheduleConfig
 
-Parameters to create a new schedule
-
 Name | Description
 ---- | -----------
-`enabled` | `boolean`<br>Whether the schedule should be enabled
-`condition` | `string`<br>Condition for whether the schedule should be enabled. When a condition is provided the result of the expression will be logically AND with the  property to decide if the schedule should be enabled.
-`maxActive` | `integer`<br>Maximum number of builds that can be active at once
-`maxChanges` | `integer`<br>Maximum number of changes the schedule can fall behind head revision. If greater than zero, builds will be triggered for every submitted changelist until the backlog is this size.
-`requireSubmittedChange` | `boolean`<br>Whether the build requires a change to be submitted
-`gate` | [ScheduleGateConfig](#schedulegateconfig)<br>Gate allowing the schedule to trigger
-`commits` | `string[]`<br>Commit tags for this schedule
-`claims` | [ScheduleClaimConfig](#scheduleclaimconfig)`[]`<br>Roles to impersonate for this schedule
-`filter` | [ChangeContentFlags](#changecontentflags-enum)`[]`<br>The types of changes to run for
-`files` | `string[]`<br>Files that should cause the job to trigger
-`templateParameters` | `string` `->` `string`<br>Parameters for the template
-`patterns` | [SchedulePatternConfig](#schedulepatternconfig)`[]`<br>New patterns for the schedule
+`enabled` | `boolean`<br>
+`condition` | `string`<br>
+`maxActive` | `integer`<br>
+`maxChanges` | `integer`<br>
+`requireSubmittedChange` | `boolean`<br>
+`gate` | [ScheduleGateConfig](#schedulegateconfig)<br>
+`commits` | `string[]`<br>
+`claims` | [ScheduleClaimConfig](#scheduleclaimconfig)`[]`<br>
+`filter` | [ChangeContentFlags](#changecontentflags-enum)`[]`<br>
+`files` | `string[]`<br>
+`templateParameters` | `string` `->` `string`<br>
+`patterns` | [SchedulePatternConfig](#schedulepatternconfig)`[]`<br>
 
 ## ScheduleGateConfig
 
-Gate allowing a schedule to trigger.
-
 Name | Description
 ---- | -----------
-`templateId` | `string`<br>The template containing the dependency
-`target` | `string`<br>Target to wait for
+`templateId` | `string`<br>
+`target` | `string`<br>
 
 ## ScheduleClaimConfig
 
-Claim to grant to leases running a scheduled job
-
 Name | Description
 ---- | -----------
-`type` | `string`<br>The claim type
-`value` | `string`<br>The claim value
+`type` | `string`<br>
+`value` | `string`<br>
 
 ## ChangeContentFlags (Enum)
 
-Flags identifying content of a changelist
-
 Name | Description
 ---- | -----------
-`ContainsCode` | The change contains code
-`ContainsContent` | The change contains content
+`ContainsCode` | 
+`ContainsContent` | 
 
 ## SchedulePatternConfig
 
-Parameters to create a new schedule
-
 Name | Description
 ---- | -----------
-`daysOfWeek` | [DayOfWeek](#dayofweek-enum)`[]`<br>Days of the week to run this schedule on. If null, the schedule will run every day.
-`minTime` | `string`<br>Time during the day for the first schedule to trigger. Measured in minutes from midnight.
-`maxTime` | `string`<br>Time during the day for the last schedule to trigger. Measured in minutes from midnight.
-`interval` | `string`<br>Interval between each schedule triggering
+`daysOfWeek` | [DayOfWeek](#dayofweek-enum)`[]`<br>
+`minTime` | `string`<br>
+`maxTime` | `string`<br>
+`interval` | `string`<br>
 
 ## DayOfWeek (Enum)
 
@@ -320,54 +286,44 @@ Name | Description
 
 ## ChainedJobTemplateConfig
 
-Trigger for another template
-
 Name | Description
 ---- | -----------
-`trigger` | `string`<br>Name of the target that needs to complete before starting the other template
-`templateId` | `string`<br>Id of the template to trigger
-`useDefaultChangeForTemplate` | `boolean`<br>Whether to use the default change for the template rather than the change for the parent job.
+`trigger` | `string`<br>
+`templateId` | `string`<br>
+`useDefaultChangeForTemplate` | `boolean`<br>
 
 ## AclConfig
 
-Parameters to update an ACL
-
 Name | Description
 ---- | -----------
-`entries` | [AclEntryConfig](#aclentryconfig)`[]`<br>Entries to replace the existing ACL
-`profiles` | [AclProfileConfig](#aclprofileconfig)`[]`<br>Defines profiles which allow grouping sets of actions into named collections
-`inherit` | `boolean`<br>Whether to inherit permissions from the parent ACL
-`exceptions` | `string[]`<br>List of exceptions to the inherited setting
+`entries` | [AclEntryConfig](#aclentryconfig)`[]`<br>
+`profiles` | [AclProfileConfig](#aclprofileconfig)`[]`<br>
+`inherit` | `boolean`<br>
+`exceptions` | `string[]`<br>
 
 ## AclEntryConfig
 
-Individual entry in an ACL
-
 Name | Description
 ---- | -----------
-`claim` | [AclClaimConfig](#aclclaimconfig)<br>Name of the user or group
-`actions` | `string[]`<br>Array of actions to allow
-`profiles` | `string[]`<br>List of profiles to grant
+`claim` | [AclClaimConfig](#aclclaimconfig)<br>
+`actions` | `string[]`<br>
+`profiles` | `string[]`<br>
 
 ## AclClaimConfig
 
-New claim to create
-
 Name | Description
 ---- | -----------
-`type` | `string`<br>The claim type
-`value` | `string`<br>The claim value
+`type` | `string`<br>
+`value` | `string`<br>
 
 ## AclProfileConfig
 
-Configuration for an ACL profile. This defines a preset group of actions which can be given to a user via an ACL entry.
-
 Name | Description
 ---- | -----------
-`id` | `string`<br>Identifier for this profile
-`actions` | `string[]`<br>Actions to include
-`excludeActions` | `string[]`<br>Actions to exclude from the inherited actions
-`extends` | `string[]`<br>Other profiles to extend from
+`id` | `string`<br>
+`actions` | `string[]`<br>
+`excludeActions` | `string[]`<br>
+`extends` | `string[]`<br>
 
 ## Priority (Enum)
 
@@ -385,32 +341,28 @@ Name | Description
 
 ## TemplateTextParameterConfig
 
-Free-form text entry parameter
-
 Name | Description
 ---- | -----------
 `type` | Text<br>Type discriminator
-`id` | `string`<br>Identifier for this parameter
-`label` | `string`<br>Name of the parameter associated with this parameter.
-`argument` | `string`<br>Argument to pass to the executor
-`default` | `string`<br>Default value for this argument
-`scheduleOverride` | `string`<br>Override for the default value for this parameter when running a scheduled build
-`hint` | `string`<br>Hint text for this parameter
-`validation` | `string`<br>Regex used to validate this parameter
-`validationError` | `string`<br>Message displayed if validation fails, informing user of valid values.
-`toolTip` | `string`<br>Tool-tip text to display
+`id` | `string`<br>
+`label` | `string`<br>
+`argument` | `string`<br>
+`default` | `string`<br>
+`scheduleOverride` | `string`<br>
+`hint` | `string`<br>
+`validation` | `string`<br>
+`validationError` | `string`<br>
+`toolTip` | `string`<br>
 
 ## TemplateListParameterConfig
 
-Allows the user to select a value from a constrained list of choices
-
 Name | Description
 ---- | -----------
 `type` | List<br>Type discriminator
-`label` | `string`<br>Label to display next to this parameter. Defaults to the parameter name.
-`style` | [TemplateListParameterStyle](#templatelistparameterstyle-enum)<br>The type of list parameter
-`items` | [TemplateListParameterItemConfig](#templatelistparameteritemconfig)`[]`<br>List of values to display in the list
-`toolTip` | `string`<br>Tool tip text to display
+`label` | `string`<br>
+`style` | [TemplateListParameterStyle](#templatelistparameterstyle-enum)<br>
+`items` | [TemplateListParameterItemConfig](#templatelistparameteritemconfig)`[]`<br>
+`toolTip` | `string`<br>
 
 ## TemplateListParameterStyle (Enum)
 
@@ -424,120 +376,104 @@ Name | Description
 
 ## TemplateListParameterItemConfig
 
-Possible option for a list parameter
-
 Name | Description
 ---- | -----------
-`id` | `string`<br>Identifier for this parameter
-`group` | `string`<br>Optional group heading to display this entry under, if the picker style supports it.
-`text` | `string`<br>Name of the parameter associated with this list.
-`argumentIfEnabled` | `string`<br>Argument to pass with this parameter.
-`argumentsIfEnabled` | `string[]`<br>Arguments to pass with this parameter.
-`argumentIfDisabled` | `string`<br>Argument to pass with this parameter.
-`argumentsIfDisabled` | `string[]`<br>Arguments to pass if this parameter is disabled.
-`default` | `boolean`<br>Whether this item is selected by default
-`scheduleOverride` | `boolean`<br>Overridden value for this property in schedule builds
+`id` | `string`<br>
+`group` | `string`<br>
+`text` | `string`<br>
+`argumentIfEnabled` | `string`<br>
+`argumentsIfEnabled` | `string[]`<br>
+`argumentIfDisabled` | `string`<br>
+`argumentsIfDisabled` | `string[]`<br>
+`default` | `boolean`<br>
+`scheduleOverride` | `boolean`<br>
 
 ## TemplateBoolParameterConfig
 
-Allows the user to toggle an option on or off
-
 Name | Description
 ---- | -----------
 `type` | Bool<br>Type discriminator
-`id` | `string`<br>Identifier for this parameter
-`label` | `string`<br>Name of the parameter associated with this parameter.
-`argumentIfEnabled` | `string`<br>Argument to add if this parameter is enabled
-`argumentsIfEnabled` | `string[]`<br>Argument to add if this parameter is enabled
-`argumentIfDisabled` | `string`<br>Argument to add if this parameter is enabled
-`argumentsIfDisabled` | `string[]`<br>Arguments to add if this parameter is disabled
-`default` | `boolean`<br>Whether this argument is enabled by default
-`scheduleOverride` | `boolean`<br>Override for this parameter in scheduled builds
-`toolTip` | `string`<br>Tool tip text to display
+`id` | `string`<br>
+`label` | `string`<br>
+`argumentIfEnabled` | `string`<br>
+`argumentsIfEnabled` | `string[]`<br>
+`argumentIfDisabled` | `string`<br>
+`argumentsIfDisabled` | `string[]`<br>
+`default` | `boolean`<br>
+`scheduleOverride` | `boolean`<br>
+`toolTip` | `string`<br>
 
 ## StreamAliasConfig
 
-Configuration for a stream alias, which represents a named subset of templates
-
 Name | Description
 ---- | -----------
-`name` | `string`<br>The unique name/identifier for this alias
-`templates` | `string[]`<br>The subset of template IDs available through this alias. These must be a subset of the enclosing IStream's Templates.
+`name` | `string`<br>
+`templates` | `string[]`<br>
 
 ## ReplicatorConfig
 
-Configuration for a stream replicator
-
 Name | Description
 ---- | -----------
-`id` | `string`<br>Identifier for the replicator within the current stream
-`enabled` | `boolean`<br>Whether the replicator is enabled
-`minChange` | `integer`<br>Minimum change number to replicate
-`maxChange` | `integer`<br>Maximum change number to replicate
-`includeContent` | `boolean`<br>Whether to include content in the replication, or just metadata
-`namespaceId` | `string`<br>Namespace to replicate data to
+`id` | `string`<br>
+`enabled` | `boolean`<br>
+`minChange` | `integer`<br>
+`maxChange` | `integer`<br>
+`includeContent` | `boolean`<br>
+`namespaceId` | `string`<br>
 
 ## WorkflowConfig
 
-Configuration for an issue workflow
-
-Name | Description
----- | -----------
-`id` | `string`<br>Identifier for this workflow
-`reportTimes` | `string[]`<br>Times of day at which to send a report
-`summaryTab` | `string`<br>Name of the tab to post summary data to
-`reportChannel` | `string`<br>Channel to post summary information for these templates.
-`reportWarnings` | `boolean`<br>Whether to include issues with a warning status in the summary
-`groupIssuesByTemplate` | `boolean`<br>Whether to group issues by template in the report
-`triageChannel` | `string`<br>Channel to post threads for triaging new issues
-`triagePrefix` | `string`<br>Prefix for all triage messages
-`triageSuffix` | `string`<br>Suffix for all triage messages
-`triageInstructions` | `string`<br>Instructions posted to triage threads
-`triageAlias` | `string`<br>User id of a Slack user/alias to ping if there is nobody assigned to an issue by default.
-`triageErrors` | `boolean`<br>Whether to include issues with an error status in the triage
-`triageWarnings` | `boolean`<br>Whether to include issues with a warning status in the triage
-`triageTypeAliases` | `string` `->` `string`<br>Slack user/alias to ping for specific issue types (such as Systemic), if there is nobody assigned to an issue by default.
-`escalateAlias` | `string`<br>Alias to ping if an issue has not been resolved for a certain amount of time
-`escalateTimes` | `integer[]`<br>Times after an issue has been opened to escalate to the alias above, in minutes. Continues to notify on the last interval once reaching the end of the list.
-`escalateRootCauseAnalysisTimes` | `integer[]`<br>Times after a manual root cause analysis has been opened to escalate to the owner, in minutes. Continues to notify on the last interval once reaching the end of the list.
-`autoRcaEnabled` | `boolean`<br>Whether to automatically run RCA processing for issues in this workflow
-`maxMentions` | `integer`<br>Maximum number of people to mention on a triage thread
-`allowMentions` | `boolean`<br>Whether to mention people on this thread. Useful to disable for testing.
-`inviteRestrictedUsers` | `boolean`<br>Uses the admin.conversations.invite API to invite users to the channel
-`skipWhenEmpty` | `boolean`<br>Skips sending reports when there are no active issues.
-`showMergeWarnings` | `boolean`<br>Whether to show warnings about merging changes into the origin stream.
-`annotations` | `string` `->` `string`<br>Additional node annotations implicit in this workflow
-`externalIssues` | [WorkflowExternalIssuesConfig](#workflowexternalissuesconfig)<br>External issue tracking configuration for this workflow
-`issueHandlers` | `string[]`<br>Additional issue handlers enabled for this workflow
+Name | Description
+---- | -----------
+`id` | `string`<br>
+`reportTimes` | `string[]`<br>
+`summaryTab` | `string`<br>
+`reportChannel` | `string`<br>
+`reportWarnings` | `boolean`<br>
+`groupIssuesByTemplate` | `boolean`<br>
+`triageChannel` | `string`<br>
+`triagePrefix` | `string`<br>
+`triageSuffix` | `string`<br>
+`triageInstructions` | `string`<br>
+`triageAlias` | `string`<br>
+`triageErrors` | `boolean`<br>
+`triageWarnings` | `boolean`<br>
+`triageTypeAliases` | `string` `->` `string`<br>
+`escalateAlias` | `string`<br>
+`escalateTimes` | `integer[]`<br>
+`escalateRootCauseAnalysisTimes` | `integer[]`<br>
+`autoRcaEnabled` | `boolean`<br>
+`maxMentions` | `integer`<br>
+`allowMentions` | `boolean`<br>
+`inviteRestrictedUsers` | `boolean`<br>
+`skipWhenEmpty` | `boolean`<br>
+`showMergeWarnings` | `boolean`<br>
+`annotations` | `string` `->` `string`<br>
+`externalIssues` | [WorkflowExternalIssuesConfig](#workflowexternalissuesconfig)<br>
+`issueHandlers` | `string[]`<br>
 
 ## WorkflowExternalIssuesConfig
 
-External issue tracking configuration for a workflow
-
 Name | Description
 ---- | -----------
-`projectKey` | `string`<br>Project key in external issue tracker
-`defaultComponentId` | `string`<br>Default component id for issues using workflow
-`defaultIssueTypeId` | `string`<br>Default issue type id for issues using workflow
+`projectKey` | `string`<br>
+`defaultComponentId` | `string`<br>
+`defaultIssueTypeId` | `string`<br>
 
 ## ArtifactTypeConfig
 
-Configuration for an artifact
-
 Name | Description
 ---- | -----------
-`name` | `string`<br>Legacy 'Name' property
-`type` | `string`<br>Name of the artifact type
-`acl` | [AclConfig](#aclconfig)<br>Acl for the artifact type
-`keepCount` | `integer`<br>Number of artifacts to retain
-`keepDays` | `integer`<br>Number of days to retain artifacts of this type
-`namespaceId` | `string`<br>Storage namespace to use for this artifact types
+`name` | `string`<br>
+`type` | `string`<br>
+`acl` | [AclConfig](#aclconfig)<br>
+`keepCount` | `integer`<br>
+`keepDays` | `integer`<br>
+`namespaceId` | `string`<br>
 
 ## StreamTag
 
-Tags for a stream
-
 Name | Description
 ---- | -----------
-`name` | `string`<br>Name of the tag
-`enabled` | `boolean`<br>Whether the tag should be enabled
+`name` | `string`<br>
+`enabled` | `boolean`<br>
diff --git a/Engine/Source/Programs/Horde/Docs/Config/Schema/Telemetry.md b/Engine/Source/Programs/Horde/Docs/Config/Schema/Telemetry.md
--- a/Engine/Source/Programs/Horde/Docs/Config/Schema/Telemetry.md
+++ b/Engine/Source/Programs/Horde/Docs/Config/Schema/Telemetry.md
@@ -2,86 +2,72 @@
 
 # *.telemetry.json
 
-Config for metrics
-
 Name | Description
 ---- | -----------
-`id` | `string`<br>Identifier for this store
-`acl` | [AclConfig](#aclconfig)<br>Permissions for this store
-`metrics` | [MetricConfig](#metricconfig)`[]`<br>Metrics to aggregate on the Horde server
-`views` | [TelemetryViewConfig](#telemetryviewconfig)`[]`<br>Configuration for telemetry views
-`include` | [ConfigInclude](#configinclude)`[]`<br>Includes for other configuration files
-`macros` | [ConfigMacro](#configmacro)`[]`<br>Macros within this configuration
+`id` | `string`<br>
+`acl` | [AclConfig](#aclconfig)<br>
+`metrics` | [MetricConfig](#metricconfig)`[]`<br>
+`views` | [TelemetryViewConfig](#telemetryviewconfig)`[]`<br>
+`include` | [ConfigInclude](#configinclude)`[]`<br>
+`macros` | [ConfigMacro](#configmacro)`[]`<br>
 
 ## AclConfig
 
-Parameters to update an ACL
-
 Name | Description
 ---- | -----------
-`entries` | [AclEntryConfig](#aclentryconfig)`[]`<br>Entries to replace the existing ACL
-`profiles` | [AclProfileConfig](#aclprofileconfig)`[]`<br>Defines profiles which allow grouping sets of actions into named collections
-`inherit` | `boolean`<br>Whether to inherit permissions from the parent ACL
-`exceptions` | `string[]`<br>List of exceptions to the inherited setting
+`entries` | [AclEntryConfig](#aclentryconfig)`[]`<br>
+`profiles` | [AclProfileConfig](#aclprofileconfig)`[]`<br>
+`inherit` | `boolean`<br>
+`exceptions` | `string[]`<br>
 
 ## AclEntryConfig
 
-Individual entry in an ACL
-
 Name | Description
 ---- | -----------
-`claim` | [AclClaimConfig](#aclclaimconfig)<br>Name of the user or group
-`actions` | `string[]`<br>Array of actions to allow
-`profiles` | `string[]`<br>List of profiles to grant
+`claim` | [AclClaimConfig](#aclclaimconfig)<br>
+`actions` | `string[]`<br>
+`profiles` | `string[]`<br>
 
 ## AclClaimConfig
 
-New claim to create
-
 Name | Description
 ---- | -----------
-`type` | `string`<br>The claim type
-`value` | `string`<br>The claim value
+`type` | `string`<br>
+`value` | `string`<br>
 
 ## AclProfileConfig
 
-Configuration for an ACL profile. This defines a preset group of actions which can be given to a user via an ACL entry.
-
 Name | Description
 ---- | -----------
-`id` | `string`<br>Identifier for this profile
-`actions` | `string[]`<br>Actions to include
-`excludeActions` | `string[]`<br>Actions to exclude from the inherited actions
-`extends` | `string[]`<br>Other profiles to extend from
+`id` | `string`<br>
+`actions` | `string[]`<br>
+`excludeActions` | `string[]`<br>
+`extends` | `string[]`<br>
 
 ## MetricConfig
 
-Configures a metric to aggregate on the server
-
 Name | Description
 ---- | -----------
-`id` | `string`<br>Identifier for this metric
-`filter` | `string`<br>Filter expression to evaluate to determine which events to include. This query is evaluated against an array.
-`property` | `string`<br>Property to aggregate
-`groupBy` | `string`<br>Property to group by. Specified as a comma-separated list of JSON path expressions.
-`function` | [AggregationFunction](#aggregationfunction-enum)<br>How to aggregate samples for this metric
-`percentile` | `integer`<br>For the percentile function, specifies the percentile to measure
-`topN` | `integer`<br>TopN aggregation
-`bottomN` | `integer`<br>BottomN aggregation
-`interval` | `string`<br>Interval for each metric. Supports times such as "2d", "1h", "1h30m", "20s".
+`id` | `string`<br>
+`filter` | `string`<br>
+`property` | `string`<br>
+`groupBy` | `string`<br>
+`function` | [AggregationFunction](#aggregationfunction-enum)<br>
+`percentile` | `integer`<br>
+`topN` | `integer`<br>
+`bottomN` | `integer`<br>
+`interval` | `string`<br>
 
 ## AggregationFunction (Enum)
 
-Method for aggregating samples into a metric
-
 Name | Description
 ---- | -----------
-`Count` | Count the number of matching elements
-`Min` | Take the minimum value of all samples
-`Max` | Take the maximum value of all samples
-`Sum` | Sum all the reported values
-`Average` | Average all the samples
-`Percentile` | Estimates the value at a certain percentile
+`Count` | 
+`Min` | 
+`Max` | 
+`Sum` | 
+`Average` | 
+`Percentile` | 
 
 ## TelemetryViewConfig
 
@@ -158,17 +144,13 @@ Name | Description
 
 ## ConfigInclude
 
-Directive to merge config data from another source
-
 Name | Description
 ---- | -----------
-`path` | `string`<br>Path to the config data to be included. May be relative to the including file's location.
+`path` | `string`<br>
 
 ## ConfigMacro
 
-Declares a config macro
-
 Name | Description
 ---- | -----------
-`name` | `string`<br>Name of the macro property
-`value` | `string`<br>Value for the macro property
+`name` | `string`<br>
+`value` | `string`<br>
diff --git a/Engine/Source/Programs/Horde/Docs/Deployment/AgentSettings.md b/Engine/Source/Programs/Horde/Docs/Deployment/AgentSettings.md
--- a/Engine/Source/Programs/Horde/Docs/Deployment/AgentSettings.md
+++ b/Engine/Source/Programs/Horde/Docs/Deployment/AgentSettings.md
@@ -6,53 +6,49 @@ All Horde-specific settings are stored in a root object called `Horde`. Other .N
 
 Name | Description
 ---- | -----------
-`serverProfiles` | `string` `->` [ServerProfile](#serverprofile)<br>Known servers to connect to
-`server` | `string`<br>The default server, unless overridden from the command line
-`name` | `string`<br>Name of agent to report as when connecting to server. By default, the computer's hostname will be used.
-`mode` | [AgentMode](#agentmode-enum)<br>Mode of operation for the agent - For trusted agents in controlled environments (e.g., build farms). These agents handle all lease types and run exclusively Horde workloads.<br>- For low-trust workstations, uses interactive authentication (human logs in). These agents yield to non-Horde workloads and only support compute leases for remote execution.
-`installed` | `boolean`<br>Whether the server is running in 'installed' mode. In this mode, on Windows, the default data directory will use the common application data folder (C:\ProgramData\Epic\Horde), and configuration data will be read from here and the registry. This setting is overridden to false for local builds from appsettings.Local.json.
-`ephemeral` | `boolean`<br>Whether agent should register as being ephemeral. Doing so will not persist any long-lived data on the server and once disconnected it's assumed to have been deleted permanently. Ideal for short-lived agents, such as spot instances on AWS EC2.
-`workingDir` | [DirectoryReference](#directoryreference)<br>Working directory for leases and jobs (i.e where files from Perforce will be checked out)
-`logsDir` | [DirectoryReference](#directoryreference)<br>Directory where agent and lease logs are written
-`shareMountingEnabled` | `boolean`<br>Whether to mount the specified list of network shares
-`shares` | [MountNetworkShare](#mountnetworkshare)`[]`<br>List of network shares to mount
-`wineExecutablePath` | `string`<br>Path to Wine executable. If null, execution under Wine is disabled
-`containerEngineExecutablePath` | `string`<br>Path to container engine executable, such as /usr/bin/podman. If null, execution of compute workloads inside a container is disabled
-`writeStepOutputToLogger` | `boolean`<br>Whether to write step output to the logging device
-`enableAwsEc2Support` | `boolean`<br>Queries information about the current agent through the AWS EC2 interface
-`useLocalStorageClient` | `boolean`<br>Option to use a local storage client rather than connecting through the server. Primarily for convenience when debugging / iterating locally.
-`computeIp` | `string`<br>Incoming IP for listening for compute work. If not set, it will be automatically resolved.
-`computePort` | `integer`<br>Incoming port for listening for compute work. Needs to be tied with a lease. Set port to 0 to disable incoming compute requests.
-`openTelemetry` | [OpenTelemetrySettings](#opentelemetrysettings)<br>Options for OpenTelemetry
-`enableTelemetry` | `boolean`<br>Whether to send telemetry back to Horde server
-`telemetryReportInterval` | `integer`<br>How often to report telemetry events to server in milliseconds
-`bundleCacheSize` | `integer`<br>Maximum size of the bundle cache, in megabytes.
-`cpuCount` | `integer`<br>Maximum number of logical CPU cores workloads should use Currently this is only provided as a hint and requires leases to respect this value as it's set via an env variable (UE_HORDE_CPU_COUNT).
-`cpuMultiplier` | `number`<br>CPU core multiplier applied to CPU core count setting For example, 32 CPU cores and a multiplier of 0.5 results in max 16 CPU usage.
-`ramGb` | `integer`<br>Maximum available RAM in gigabytes. Defaults to available RAM.
-`properties` | `string` `->` `string`<br>Key/value properties in addition to those set internally by the agent
-`adminEndpoints` | `string[]`<br>Listen addresses for the built-in HTTP admin/management server. Disabled when empty. If activated, it's recommended to bind only to localhost for security reasons. Example: localhost:7008 to listen on localhost, port 7008
-`healthCheckEndpoints` | `string[]`<br>Listen addresses for the built-in HTTP health check server. Disabled when empty. If activated, it's recommended to bind only to localhost for security reasons. Example: *:7009 to listen on all interfaces/IPs, port 7009 If all interfaces are bound with *, make sure to run process as administrator.
-`tempStorageMaxBatchSize` | `integer`<br>Maximum number of temp storage exports to write in a single request
+`serverProfiles` | `string` `->` [ServerProfile](#serverprofile)<br>
+`server` | `string`<br>
+`name` | `string`<br>
+`mode` | [AgentMode](#agentmode-enum)<br>
+`installed` | `boolean`<br>
+`ephemeral` | `boolean`<br>
+`workingDir` | [DirectoryReference](#directoryreference)<br>
+`logsDir` | [DirectoryReference](#directoryreference)<br>
+`shareMountingEnabled` | `boolean`<br>
+`shares` | [MountNetworkShare](#mountnetworkshare)`[]`<br>
+`wineExecutablePath` | `string`<br>
+`containerEngineExecutablePath` | `string`<br>
+`writeStepOutputToLogger` | `boolean`<br>
+`enableAwsEc2Support` | `boolean`<br>
+`useLocalStorageClient` | `boolean`<br>
+`computeIp` | `string`<br>
+`computePort` | `integer`<br>
+`openTelemetry` | [OpenTelemetrySettings](#opentelemetrysettings)<br>
+`enableTelemetry` | `boolean`<br>
+`telemetryReportInterval` | `integer`<br>
+`bundleCacheSize` | `integer`<br>
+`cpuCount` | `integer`<br>
+`cpuMultiplier` | `number`<br>
+`ramGb` | `integer`<br>
+`properties` | `string` `->` `string`<br>
+`adminEndpoints` | `string[]`<br>
+`healthCheckEndpoints` | `string[]`<br>
+`tempStorageMaxBatchSize` | `integer`<br>
 
 ## ServerProfile
 
-Information about a server to use
-
 Name | Description
 ---- | -----------
-`name` | `string`<br>Name of this server profile
-`environment` | `string`<br>Name of the environment (currently just used for tracing)
-`url` | `string`<br>Url of the server
-`token` | `string`<br>Bearer token to use to initiate the connection
-`useInteractiveAuth` | `boolean`<br>Whether to authenticate interactively in a desktop environment (for example, when agent is running on a user's workstation)
-`thumbprint` | `string`<br>Thumbprint of a certificate to trust. Allows using self-signed certs for the server.
-`thumbprints` | `string[]`<br>Thumbprints of certificates to trust. Allows using self-signed certs for the server.
+`name` | `string`<br>
+`environment` | `string`<br>
+`url` | `string`<br>
+`token` | `string`<br>
+`useInteractiveAuth` | `boolean`<br>
+`thumbprint` | `string`<br>
+`thumbprints` | `string[]`<br>
 
 ## AgentMode (Enum)
 
-Defines the operation mode of the agent Duplicated to prevent depending on Protobuf structures (weak name reference when deserializing JSON)
-
 Name | Description
 ---- | -----------
 `Dedicated` | 
@@ -67,12 +63,10 @@ Name | Description
 
 ## MountNetworkShare
 
-Describes a network share to mount
-
 Name | Description
 ---- | -----------
-`mountPoint` | `string`<br>Where the share should be mounted on the local machine. Must be a drive letter for Windows.
-`remotePath` | `string`<br>Path to the remote resource
+`mountPoint` | `string`<br>
+`remotePath` | `string`<br>
 
 ## OpenTelemetrySettings
 
diff --git a/Engine/Source/Programs/Horde/Docs/Deployment/ServerSettings.md b/Engine/Source/Programs/Horde/Docs/Deployment/ServerSettings.md
--- a/Engine/Source/Programs/Horde/Docs/Deployment/ServerSettings.md
+++ b/Engine/Source/Programs/Horde/Docs/Deployment/ServerSettings.md
@@ -74,13 +74,11 @@ Name | Description
 
 ## RunMode (Enum)
 
-Type of run mode this process should use. Each carry different types of workloads. More than one mode can be active. But not all modes are not guaranteed to be compatible with each other and will raise an error if combined in such a way.
-
 Name | Description
 ---- | -----------
-`None` | Default no-op value (ASP.NET config will default to this for enums that cannot be parsed)
-`Server` | Handle and respond to incoming external requests, such as HTTP REST and gRPC calls. These requests are time-sensitive and short-lived, typically less than 5 secs. If processes handling requests are unavailable, it will be very visible for users.
-`Worker` | Run non-request facing workloads. Such as background services, processing queues, running work based on timers etc. Short periods of downtime or high CPU usage due to bursts are fine for this mode. No user requests will be impacted directly. If auto-scaling is used, a much more aggressive policy can be applied (tighter process packing, higher avg CPU usage).
+`None` | 
+`Server` | 
+`Worker` | 
 
 ## AuthMethod (Enum)
 
@@ -127,85 +125,106 @@ Name | Description
 
 Name | Description
 ---- | -----------
+`analytics` | [AnalyticsServerConfig](#analyticsserverconfig)<br>Configuration for the analytics plugin
 `build` | [BuildServerConfig](#buildserverconfig)<br>Configuration for the build plugin
 `compute` | [ComputeServerConfig](#computeserverconfig)<br>Configuration for the compute plugin
-`analytics` | [AnalyticsServerConfig](#analyticsserverconfig)<br>Configuration for the analytics plugin
-`exampleplugin` | [ExamplePluginServerConfig](#examplepluginserverconfig)<br>Configuration for the exampleplugin plugin
-`secrets` | [SecretsServerConfig](#secretsserverconfig)<br>Configuration for the secrets plugin
 `experimental` | [ExperimentalServerConfig](#experimentalserverconfig)<br>Configuration for the experimental plugin
+`secrets` | [SecretsServerConfig](#secretsserverconfig)<br>Configuration for the secrets plugin
 `storage` | [StorageServerConfig](#storageserverconfig)<br>Configuration for the storage plugin
-`symbols` | [PluginServerConfig](#pluginserverconfig)<br>Configuration for the symbols plugin
+`exampleplugin` | [ExamplePluginServerConfig](#examplepluginserverconfig)<br>Configuration for the exampleplugin plugin
 `tools` | [ToolsServerConfig](#toolsserverconfig)<br>Configuration for the tools plugin
+`symbols` | [PluginServerConfig](#pluginserverconfig)<br>Configuration for the symbols plugin
+
+## AnalyticsServerConfig
+
+Name | Description
+---- | -----------
+`sinks` | [TelemetrySinkConfig](#telemetrysinkconfig)<br>
+`enabled` | `boolean`<br>
+
+## TelemetrySinkConfig
+
+Name | Description
+---- | -----------
+`epic` | [EpicTelemetryConfig](#epictelemetryconfig)<br>
+`mongo` | [MongoTelemetryConfig](#mongotelemetryconfig)<br>
+
+## EpicTelemetryConfig
+
+Name | Description
+---- | -----------
+`url` | `string`<br>
+`appId` | `string`<br>
+`enabled` | `boolean`<br>
+
+## MongoTelemetryConfig
+
+Name | Description
+---- | -----------
+`retainDays` | `number`<br>
+`enabled` | `boolean`<br>
 
 ## BuildServerConfig
 
-Static configuration for the build plugin
-
-Name | Description
----- | -----------
-`perforce` | [PerforceConnectionSettings](#perforceconnectionsettings)`[]`<br>Perforce connections for use by the Horde server (not agents)
-`useLocalPerforceEnv` | `boolean`<br>Whether to use the local Perforce environment
-`perforceConnectionPoolSize` | `integer`<br>Number of pooled perforce connections to keep
-`enableConformTasks` | `boolean`<br>Whether to enable the conform task source.
-`p4SwarmUrl` | `string`<br>Url of P4 Swarm installation
-`robomergeUrl` | `string`<br>Url of Robomergem installation
-`commitsViewerUrl` | `string`<br>Url of Commits Viewer
-`jiraUsername` | `string`<br>The Jira service account user name
-`jiraApiToken` | `string`<br>The Jira service account API token
-`jiraUrl` | `string`<br>The Uri for the Jira installation
-`sharedDeviceCheckoutDays` | `integer`<br>The number of days shared device checkouts are held
-`deviceProblemCooldownMinutes` | `integer`<br>The number of cooldown minutes for device problems
-`deviceReportChannel` | `string`<br>Channel to send device reports to
-`disableSchedules` | `boolean`<br>Whether to run scheduled jobs.
-`slackToken` | `string`<br>Bot token for interacting with Slack (xoxb-*)
-`slackSocketToken` | `string`<br>Token for opening a socket to slack (xapp-*)
-`slackAdminToken` | `string`<br>Admin user token for Slack (xoxp-*). This is only required when using the admin endpoints to invite users.
-`slackUsers` | `string`<br>Filtered list of slack users to send notifications to. Should be Slack user ids, separated by commas.
-`slackErrorPrefix` | `string`<br>Prefix to use when reporting errors
-`slackWarningPrefix` | `string`<br>Prefix to use when reporting warnings
-`configNotificationChannel` | `string`<br>Channel for sending messages related to config update failures
-`updateStreamsNotificationChannel` | `string`<br>Channel to send stream notification update failures to
-`jobNotificationChannel` | `string`<br>Slack channel to send job related notifications to. Multiple channels can be specified, separated by ;
-`agentNotificationChannel` | `string`<br>Slack channel to send agent related notifications to.
-`testDataRetainMonths` | `integer`<br>The number of months to retain test data
-`testHealthHistoryEvaluationDays` | `integer`<br>The number of days to go back in historical data to evaluate test health
-`testHealthReportRemindDays` | `integer`<br>The number of days after which a test health report remind for improvement
-`blockCacheDir` | `string`<br>Directory to store the fine-grained block cache. This caches individual exports embedded in bundles.
-`blockCacheSize` | `string`<br>Maximum size of the block cache. Accepts standard binary suffixes. Currently only allocates in multiples of 1024mb.
-`blockCacheSizeBytes` | `integer`<br>Accessor for the block cache size in bytes
-`commits` | [CommitSettings](#commitsettings)<br>Options for the commit service
-`enabled` | `boolean`<br>Whether the plugin should be enabled or not
+Name | Description
+---- | -----------
+`perforce` | [PerforceConnectionSettings](#perforceconnectionsettings)`[]`<br>
+`useLocalPerforceEnv` | `boolean`<br>
+`perforceConnectionPoolSize` | `integer`<br>
+`enableConformTasks` | `boolean`<br>
+`p4SwarmUrl` | `string`<br>
+`robomergeUrl` | `string`<br>
+`commitsViewerUrl` | `string`<br>
+`jiraUsername` | `string`<br>
+`jiraApiToken` | `string`<br>
+`jiraUrl` | `string`<br>
+`sharedDeviceCheckoutDays` | `integer`<br>
+`deviceProblemCooldownMinutes` | `integer`<br>
+`deviceReportChannel` | `string`<br>
+`disableSchedules` | `boolean`<br>
+`slackToken` | `string`<br>
+`slackSocketToken` | `string`<br>
+`slackAdminToken` | `string`<br>
+`slackUsers` | `string`<br>
+`slackErrorPrefix` | `string`<br>
+`slackWarningPrefix` | `string`<br>
+`configNotificationChannel` | `string`<br>
+`updateStreamsNotificationChannel` | `string`<br>
+`jobNotificationChannel` | `string`<br>
+`agentNotificationChannel` | `string`<br>
+`testDataRetainMonths` | `integer`<br>
+`testHealthHistoryEvaluationDays` | `integer`<br>
+`testHealthReportRemindDays` | `integer`<br>
+`blockCacheDir` | `string`<br>
+`blockCacheSize` | `string`<br>
+`blockCacheSizeBytes` | `integer`<br>
+`commits` | [CommitSettings](#commitsettings)<br>
+`enabled` | `boolean`<br>
 
 ## PerforceConnectionSettings
 
-Perforce connection information for use by the Horde server (for reading config files, etc...)
-
 Name | Description
 ---- | -----------
-`id` | `string`<br>Identifier for this server
-`serverAndPort` | `string`<br>Server and port
-`credentials` | [PerforceCredentials](#perforcecredentials)<br>Credentials for the server
+`id` | `string`<br>
+`serverAndPort` | `string`<br>
+`credentials` | [PerforceCredentials](#perforcecredentials)<br>
 
 ## PerforceCredentials
 
-Credentials for a Perforce user
-
 Name | Description
 ---- | -----------
-`userName` | `string`<br>The username
-`password` | `string`<br>Password for the user
-`ticket` | `string`<br>Login ticket for the user (will be used instead of password if set)
+`userName` | `string`<br>
+`password` | `string`<br>
+`ticket` | `string`<br>
 
 ## CommitSettings
 
-Options for the commit service
-
 Name | Description
 ---- | -----------
-`replicateMetadata` | `boolean`<br>Whether to mirror commit metadata to the database
-`replicateContent` | `boolean`<br>Whether to mirror commit data to storage
-`bundle` | [BundleOptions](#bundleoptions)<br>Options for how objects are packed together
-`chunking` | [ChunkingOptions](#chunkingoptions)<br>Options for how objects are sliced
+`replicateMetadata` | `boolean`<br>
+`replicateContent` | `boolean`<br>
+`bundle` | [BundleOptions](#bundleoptions)<br>
+`chunking` | [ChunkingOptions](#chunkingoptions)<br>
 
 ## BundleOptions
 
@@ -283,247 +302,176 @@ Name | Description
 
 ## ComputeServerConfig
 
-Static configuration for the compute plugin
-
 Name | Description
 ---- | -----------
-`enableUpgradeTasks` | `boolean`<br>Whether to enable the upgrade task source, always upgrading agents to the latest version
-`withAws` | `boolean`<br>Whether to enable Amazon Web Services (AWS) specific features
-`awsRegions` | `string[]`<br>List of AWS regions for Horde to be aware of (e.g. us-east-1 or eu-central-1) Right now, this is only used for replicating CloudWatch metrics to multiple regions
-`awsAutoScalingQueueUrls` | `string[]`<br>AWS SQS queue URLs where lifecycle events from EC2 auto-scaling are received
-`fleetManagerV2` | [FleetManagerType](#fleetmanagertype-enum)<br>Default fleet manager to use (when not specified by pool)
-`fleetManagerV2Config` | `object`<br>Config for the fleet manager (serialized JSON)
-`autoEnrollAgents` | `boolean`<br>Whether to automatically enroll agents in the farm
-`defaultAgentPoolSizeStrategy` | [PoolSizeStrategy](#poolsizestrategy-enum)<br>Default agent pool sizing strategy for pools that doesn't have one explicitly configured
-`agentPoolScaleOutCooldownSeconds` | `integer`<br>Scale-out cooldown for auto-scaling agent pools (in seconds). Can be overridden by per-pool settings.
-`agentPoolScaleInCooldownSeconds` | `integer`<br>Scale-in cooldown for auto-scaling agent pools (in seconds). Can be overridden by per-pool settings.
-`computeTunnelPort` | `integer`<br>Port to listen on for tunneling compute sockets to agents
-`computeTunnelAddress` | `string`<br>What address (host:port) clients should connect to for compute socket tunneling Port may differ from  if Horde server is behind a reverse proxy/firewall
-`enabled` | `boolean`<br>Whether the plugin should be enabled or not
+`enableUpgradeTasks` | `boolean`<br>
+`withAws` | `boolean`<br>
+`awsRegions` | `string[]`<br>
+`awsAutoScalingQueueUrls` | `string[]`<br>
+`fleetManagerV2` | [FleetManagerType](#fleetmanagertype-enum)<br>
+`fleetManagerV2Config` | `object`<br>
+`autoEnrollAgents` | `boolean`<br>
+`defaultAgentPoolSizeStrategy` | [PoolSizeStrategy](#poolsizestrategy-enum)<br>
+`agentPoolScaleOutCooldownSeconds` | `integer`<br>
+`agentPoolScaleInCooldownSeconds` | `integer`<br>
+`computeTunnelPort` | `integer`<br>
+`computeTunnelAddress` | `string`<br>
+`enabled` | `boolean`<br>
 
 ## FleetManagerType (Enum)
 
-Available fleet managers
-
 Name | Description
 ---- | -----------
-`Default` | Default fleet manager
-`NoOp` | No-op fleet manager.
-`Aws` | Fleet manager for handling AWS EC2 instances. Will create and/or terminate instances from scratch.
-`AwsReuse` | Fleet manager for handling AWS EC2 instances. Will start already existing but stopped instances to reuse existing EBS disks.
-`AwsRecycle` | Fleet manager for handling AWS EC2 instances. Will start already existing but stopped instances to reuse existing EBS disks.
-`AwsAsg` | Fleet manager for handling AWS EC2 instances. Uses an EC2 auto-scaling group for controlling the number of running instances.
+`Default` | 
+`NoOp` | 
+`Aws` | 
+`AwsReuse` | 
+`AwsRecycle` | 
+`AwsAsg` | 
 
 ## PoolSizeStrategy (Enum)
 
-Available pool sizing strategies
-
 Name | Description
 ---- | -----------
-`LeaseUtilization` | Strategy based on lease utilization
-`JobQueue` | Strategy based on size of job build queue
-`NoOp` | No-op strategy used as fallback/default behavior
-`ComputeQueueAwsMetric` | A no-op strategy that reports metrics to let an external AWS auto-scaling policy scale the fleet
-`LeaseUtilizationAwsMetric` | A no-op strategy that reports metrics to let an external AWS auto-scaling policy scale the fleet
-
-## AnalyticsServerConfig
+`LeaseUtilization` | 
+`JobQueue` | 
+`NoOp` | 
+`ComputeQueueAwsMetric` | 
+`LeaseUtilizationAwsMetric` | 
 
-Server configuration for the analytics system
-
-Name | Description
----- | -----------
-`sinks` | [TelemetrySinkConfig](#telemetrysinkconfig)<br>Settings for the various telemetry sinks
-`enabled` | `boolean`<br>Whether the plugin should be enabled or not
-
-## TelemetrySinkConfig
-
-Telemetry sinks
-
-Name | Description
----- | -----------
-`epic` | [EpicTelemetryConfig](#epictelemetryconfig)<br>Settings for the Epic telemetry sink
-`mongo` | [MongoTelemetryConfig](#mongotelemetryconfig)<br>Settings for the MongoDB telemetry sink
-
-## EpicTelemetryConfig
-
-Configuration for the telemetry sink
-
-Name | Description
----- | -----------
-`url` | `string`<br>Base URL for the telemetry server
-`appId` | `string`<br>Application name to send in the event messages
-`enabled` | `boolean`<br>Whether to enable this sink
-
-## MongoTelemetryConfig
-
-Configuration for the telemetry sink
-
-Name | Description
----- | -----------
-`retainDays` | `number`<br>Number of days worth of telmetry events to keep
-`enabled` | `boolean`<br>Whether to enable this sink
-
-## ExamplePluginServerConfig
-
-Server configuration for the ExamplePlugin
+## ExperimentalServerConfig
 
 Name | Description
 ---- | -----------
-`enabled` | `boolean`<br>Whether the plugin should be enabled or not
+`enabled` | `boolean`<br>
 
 ## SecretsServerConfig
 
-Static configuration for the secrets plugin
-
 Name | Description
 ---- | -----------
-`withAws` | `boolean`<br>Whether to enable Amazon Web Services (AWS) specific features
-`enabled` | `boolean`<br>Whether the plugin should be enabled or not
-
-## ExperimentalServerConfig
-
-Server configuration for the experimental plugin
-
-Name | Description
----- | -----------
-`enabled` | `boolean`<br>Whether the plugin should be enabled or not
+`withAws` | `boolean`<br>
+`enabled` | `boolean`<br>
 
 ## StorageServerConfig
 
-Static settings for the storage system
-
 Name | Description
 ---- | -----------
-`bundleCacheDir` | `string`<br>Directory to use for the coarse-grained backend cache. This caches full bundles downloaded from the upstream object store.
-`bundleCacheSize` | `string`<br>Maximum size of the storage cache on disk. Accepts standard binary suffixes (kb, mb, gb, tb, etc...)
-`bundleCacheSizeBytes` | `integer`<br>Accessor for the bundle cache size in bytes
-`backends` | [BackendConfig](#backendconfig)`[]`<br>Overridden settings for storage backends. Useful for running against a production server with custom backends.
-`enabled` | `boolean`<br>Whether the plugin should be enabled or not
+`bundleCacheDir` | `string`<br>
+`bundleCacheSize` | `string`<br>
+`bundleCacheSizeBytes` | `integer`<br>
+`backends` | [BackendConfig](#backendconfig)`[]`<br>
+`enabled` | `boolean`<br>
 
 ## BackendConfig
 
-Common settings object for different providers
-
 Name | Description
 ---- | -----------
-`id` | `string`<br>The storage backend ID
-`base` | `string`<br>Base backend to copy default settings from
-`secondary` | `string`<br>Specifies another backend to read from if an object is not found in this one. Can be used when migrating data from one backend to another.
+`id` | `string`<br>
+`base` | `string`<br>
+`secondary` | `string`<br>
 `type` | [StorageBackendType](#storagebackendtype-enum)<br>
 `baseDir` | `string`<br>
-`awsBucketName` | `string`<br>Name of the bucket to use
-`awsBucketPath` | `string`<br>Base path within the bucket
-`awsCredentials` | [AwsCredentialsType](#awscredentialstype-enum)<br>Type of credentials to use
-`awsRole` | `string`<br>ARN of a role to assume
-`awsProfile` | `string`<br>The AWS profile to read credentials form
-`awsRegion` | `string`<br>Region to connect to
-`azureConnectionString` | `string`<br>Connection string for Azure
-`azureContainerName` | `string`<br>Name of the container
+`awsBucketName` | `string`<br>
+`awsBucketPath` | `string`<br>
+`awsCredentials` | [AwsCredentialsType](#awscredentialstype-enum)<br>
+`awsRole` | `string`<br>
+`awsProfile` | `string`<br>
+`awsRegion` | `string`<br>
+`azureConnectionString` | `string`<br>
+`azureContainerName` | `string`<br>
 `relayServer` | `string`<br>
 `relayToken` | `string`<br>
-`gcsBucketName` | `string`<br>Name of the GCS bucket to use
-`gcsBucketPath` | `string`<br>Base path within the bucket
+`gcsBucketName` | `string`<br>
+`gcsBucketPath` | `string`<br>
 
 ## StorageBackendType (Enum)
 
-Types of storage backend to use
-
 Name | Description
 ---- | -----------
-`FileSystem` | Local filesystem
-`Aws` | AWS S3
-`Azure` | Azure blob store
-`Gcs` | Google Cloud Storage
-`Memory` | In-memory only (for testing)
+`FileSystem` | 
+`Aws` | 
+`Azure` | 
+`Gcs` | 
+`Memory` | 
 
 ## AwsCredentialsType (Enum)
 
-Credentials to use for AWS
-
 Name | Description
 ---- | -----------
-`Default` | Use default credentials from the AWS SDK
-`Profile` | Read credentials from the  profile in the AWS config file
-`AssumeRole` | Assume a particular role. Should specify ARN in
-`AssumeRoleWebIdentity` | Assume a particular role using the current environment variables.
-
-## PluginServerConfig
+`Default` | 
+`Profile` | 
+`AssumeRole` | 
+`AssumeRoleWebIdentity` | 
 
-Base class for plugin server config objects
+## ExamplePluginServerConfig
 
 Name | Description
 ---- | -----------
-`enabled` | `boolean`<br>Whether the plugin should be enabled or not
+`enabled` | `boolean`<br>
 
 ## ToolsServerConfig
 
-Server configuration for bundled tools
-
 Name | Description
 ---- | -----------
-`bundledTools` | [BundledToolConfig](#bundledtoolconfig)`[]`<br>Tools bundled along with the server. Data for each tool can be produced using the 'bundle create' command, and should be stored in the Tools directory.
-`enabled` | `boolean`<br>Whether the plugin should be enabled or not
+`bundledTools` | [BundledToolConfig](#bundledtoolconfig)`[]`<br>
+`enabled` | `boolean`<br>
 
 ## BundledToolConfig
 
-Configuration for a tool bundled alongsize the server
-
 Name | Description
 ---- | -----------
-`version` | `string`<br>Version string for the current tool data
-`refName` | `string`<br>Ref name in the tools directory
-`dataDir` | `string`<br>Directory containing blob data for this tool. If empty, the tools/{id} folder next to the server will be used.
-`id` | `string`<br>Unique identifier for the tool
-`name` | `string`<br>Name of the tool
-`description` | `string`<br>Description for the tool
-`category` | `string`<br>Category for the tool. Will cause the tool to be shown in a different tab in the dashboard.
-`group` | `string`<br>Grouping key for different variations of the same tool. The dashboard will show these together.
-`platforms` | `string[]`<br>Platforms for this tool. Takes the form of a NET RID (https://learn.microsoft.com/en-us/dotnet/core/rid-catalog).
-`public` | `boolean`<br>Whether this tool should be exposed for download on a public endpoint without authentication
-`showInUgs` | `boolean`<br>Whether to show this tool for download in the UGS tools menu
-`showInDashboard` | `boolean`<br>Whether to show this tool for download in the dashboard
-`showInToolbox` | `boolean`<br>Whether to show this tool for download in Unreal Toolbox
-`metadata` | `string` `->` `string`<br>Metadata for this tool
-`namespaceId` | `string`<br>Default namespace for new deployments of this tool
-`acl` | [AclConfig](#aclconfig)<br>Permissions for the tool
+`version` | `string`<br>
+`refName` | `string`<br>
+`dataDir` | `string`<br>
+`id` | `string`<br>
+`name` | `string`<br>
+`description` | `string`<br>
+`category` | `string`<br>
+`group` | `string`<br>
+`platforms` | `string[]`<br>
+`public` | `boolean`<br>
+`showInUgs` | `boolean`<br>
+`showInDashboard` | `boolean`<br>
+`showInToolbox` | `boolean`<br>
+`metadata` | `string` `->` `string`<br>
+`namespaceId` | `string`<br>
+`acl` | [AclConfig](#aclconfig)<br>
 
 ## AclConfig
 
-Parameters to update an ACL
-
 Name | Description
 ---- | -----------
-`entries` | [AclEntryConfig](#aclentryconfig)`[]`<br>Entries to replace the existing ACL
-`profiles` | [AclProfileConfig](#aclprofileconfig)`[]`<br>Defines profiles which allow grouping sets of actions into named collections
-`inherit` | `boolean`<br>Whether to inherit permissions from the parent ACL
-`exceptions` | `string[]`<br>List of exceptions to the inherited setting
+`entries` | [AclEntryConfig](#aclentryconfig)`[]`<br>
+`profiles` | [AclProfileConfig](#aclprofileconfig)`[]`<br>
+`inherit` | `boolean`<br>
+`exceptions` | `string[]`<br>
 
 ## AclEntryConfig
 
-Individual entry in an ACL
-
 Name | Description
 ---- | -----------
-`claim` | [AclClaimConfig](#aclclaimconfig)<br>Name of the user or group
-`actions` | `string[]`<br>Array of actions to allow
-`profiles` | `string[]`<br>List of profiles to grant
+`claim` | [AclClaimConfig](#aclclaimconfig)<br>
+`actions` | `string[]`<br>
+`profiles` | `string[]`<br>
 
 ## AclClaimConfig
 
-New claim to create
-
 Name | Description
 ---- | -----------
-`type` | `string`<br>The claim type
-`value` | `string`<br>The claim value
+`type` | `string`<br>
+`value` | `string`<br>
 
 ## AclProfileConfig
 
-Configuration for an ACL profile. This defines a preset group of actions which can be given to a user via an ACL entry.
+Name | Description
+---- | -----------
+`id` | `string`<br>
+`actions` | `string[]`<br>
+`excludeActions` | `string[]`<br>
+`extends` | `string[]`<br>
+
+## PluginServerConfig
 
 Name | Description
 ---- | -----------
-`id` | `string`<br>Identifier for this profile
-`actions` | `string[]`<br>Actions to include
-`excludeActions` | `string[]`<br>Actions to exclude from the inherited actions
-`extends` | `string[]`<br>Other profiles to extend from
+`enabled` | `boolean`<br>
