Horde: Batch fetch agent sessions from Redis when loading multiple agents

Replace sequential per-agent Redis GET calls in AgentCollection with a single MGET, which reduces Redis round-trips.

diff --git a/Engine/Source/Programs/Horde/Plugins/Compute/HordeServer.Compute/Agents/AgentCollection.cs b/Engine/Source/Programs/Horde/Plugins/Compute/HordeServer.Compute/Agents/AgentCollection.cs
--- a/Engine/Source/Programs/Horde/Plugins/Compute/HordeServer.Compute/Agents/AgentCollection.cs
+++ b/Engine/Source/Programs/Horde/Plugins/Compute/HordeServer.Compute/Agents/AgentCollection.cs
@@ -489,6 +489,50 @@ namespace HordeServer.Agents
 		ValueTask<Agent?> CreateAgentObjectAsync(AgentDocument? document, CancellationToken cancellationToken)
 			=> CreateAgentObjectAsync(document, null, cancellationToken);
 
+		async Task<List<IAgent>> CreateAgentObjectsBatchedAsync(List<AgentDocument> documents, CancellationToken cancellationToken)
+		{
+			// Collect all session IDs for a single batched Redis MGET
+			List<SessionId> sessionIds = new (documents.Count);
+			foreach (AgentDocument document in documents)
+			{
+				if (document.SessionId != null)
+				{
+					sessionIds.Add(document.SessionId.Value);
+				}
+			}
+
+			// Batch fetch all sessions in one round-trip
+			Dictionary<SessionId, RpcSession> sessions = await _scheduler.TryGetSessionsAsync(sessionIds, cancellationToken);
+
+			// Create agent objects using pre-fetched sessions
+			List<IAgent> agents = new (documents.Count);
+			foreach (AgentDocument document in documents)
+			{
+				Agent? agent;
+				if (document.SessionId == null)
+				{
+					agent = new Agent(this, document, null);
+				}
+				else if (sessions.TryGetValue(document.SessionId.Value, out RpcSession? session))
+				{
+					// Safe to bypass the reconciliation loop in CreateAgentObjectAsync; the dictionary is keyed by
+					// document.SessionId so session.SessionId == document.SessionId is guaranteed.
+					agent = new Agent(this, document, session);
+				}
+				else
+				{
+					// Session missing from Redis — fall back to sequential handling which clears the stale session
+					agent = await CreateAgentObjectAsync(document, cancellationToken);
+				}
+
+				if (agent != null)
+				{
+					agents.Add(agent);
+				}
+			}
+			return agents;
+		}
+
 		async ValueTask<Agent?> CreateAgentObjectAsync(AgentDocument? document, RpcSession? session, CancellationToken cancellationToken)
 		{
 			if (document == null)
@@ -635,18 +679,7 @@ namespace HordeServer.Agents
 		{
 			List<AgentDocument> documents = await _agentCollection.Find(p => agentIds.Contains(p.Id)).ToListAsync(cancellationToken);
 			documents = await PostLoadAsync(documents, cancellationToken);
-
-			List<IAgent> agents = new List<IAgent>();
-			foreach (AgentDocument document in documents)
-			{
-				IAgent? agent = await CreateAgentObjectAsync(document, cancellationToken);
-				if (agent != null)
-				{
-					agents.Add(agent);
-				}
-			}
-
-			return agents;
+			return await CreateAgentObjectsBatchedAsync(documents, cancellationToken);
 		}
 
 		/// <inheritdoc/>
@@ -687,10 +720,9 @@ namespace HordeServer.Agents
 				while (await cursor.MoveNextAsync(cancellationToken))
 				{
 					List<AgentDocument> documents = await PostLoadAsync(cursor.Current.ToList(), cancellationToken);
-					foreach (AgentDocument document in documents)
+					foreach (IAgent agent in await CreateAgentObjectsBatchedAsync(documents, cancellationToken))
 					{
-						IAgent? agent = await CreateAgentObjectAsync(document, cancellationToken);
-						if (agent != null && (status == null || agent.Status == status.Value))
+						if (status == null || agent.Status == status.Value)
 						{
 							yield return agent;
 						}
diff --git a/Engine/Source/Programs/Horde/Plugins/Compute/HordeServer.Compute/Agents/AgentScheduler.cs b/Engine/Source/Programs/Horde/Plugins/Compute/HordeServer.Compute/Agents/AgentScheduler.cs
--- a/Engine/Source/Programs/Horde/Plugins/Compute/HordeServer.Compute/Agents/AgentScheduler.cs
+++ b/Engine/Source/Programs/Horde/Plugins/Compute/HordeServer.Compute/Agents/AgentScheduler.cs
@@ -296,6 +296,37 @@ namespace HordeServer.Agents
 			return await database.StringGetAsync(Keys.Sessions[sessionId].State);
 		}
 
+		/// <inheritdoc/>
+		public async Task<Dictionary<SessionId, RpcSession>> TryGetSessionsAsync(IReadOnlyList<SessionId> sessionIds, CancellationToken cancellationToken = default)
+		{
+			Dictionary<SessionId, RpcSession> result = new (sessionIds.Count);
+			if (sessionIds.Count == 0)
+			{
+				return result;
+			}
+
+			IDatabase database = _redisService.GetDatabase();
+
+			// Use .State.Inner to get the raw RedisKey; the multi-key StringGetAsync overload requires RedisKey[]
+			// rather than the typed key wrapper used by the single-key overload in TryGetSessionAsync.
+			RedisKey[] keys = new RedisKey[sessionIds.Count];
+			for (int idx = 0; idx < sessionIds.Count; idx++)
+			{
+				keys[idx] = Keys.Sessions[sessionIds[idx]].State.Inner;
+			}
+
+			RedisValue[] values = await database.StringGetAsync(keys).WaitAsync(cancellationToken);
+			for (int idx = 0; idx < sessionIds.Count; idx++)
+			{
+				if (!values[idx].IsNull)
+				{
+					result[sessionIds[idx]] = RedisSerializer.Deserialize<RpcSession>(values[idx]);
+				}
+			}
+
+			return result;
+		}
+
 		/// <inheritdoc/>
 		public async Task<RpcAgentCapabilities?> TryGetCapabilitiesAsync(RpcSession session, CancellationToken cancellationToken = default)
 		{
diff --git a/Engine/Source/Programs/Horde/Plugins/Compute/HordeServer.Compute/Agents/IAgentScheduler.cs b/Engine/Source/Programs/Horde/Plugins/Compute/HordeServer.Compute/Agents/IAgentScheduler.cs
--- a/Engine/Source/Programs/Horde/Plugins/Compute/HordeServer.Compute/Agents/IAgentScheduler.cs
+++ b/Engine/Source/Programs/Horde/Plugins/Compute/HordeServer.Compute/Agents/IAgentScheduler.cs
@@ -54,6 +54,14 @@ namespace HordeServer.Agents
 		/// <returns>Session object, or null if the session is not active</returns>
 		Task<RpcSession?> TryGetSessionAsync(SessionId sessionId, CancellationToken cancellationToken = default);
 
+		/// <summary>
+		/// Batch fetch the latest state of multiple sessions in a single round-trip
+		/// </summary>
+		/// <param name="sessionIds">The session identifiers to look up</param>
+		/// <param name="cancellationToken">Cancellation token for the operation</param>
+		/// <returns>Dictionary mapping each found session ID to its session object. Missing sessions are omitted.</returns>
+		Task<Dictionary<SessionId, RpcSession>> TryGetSessionsAsync(IReadOnlyList<SessionId> sessionIds, CancellationToken cancellationToken = default);
+
 		/// <summary>
 		/// Attempt to get capabilities for the given session
 		/// </summary>
