Horde: Add IgnorePatterns documentation

diff --git a/Engine/Source/Programs/Horde/Docs/Config/BuildHealth.md b/Engine/Source/Programs/Horde/Docs/Config/BuildHealth.md
--- a/Engine/Source/Programs/Horde/Docs/Config/BuildHealth.md
+++ b/Engine/Source/Programs/Horde/Docs/Config/BuildHealth.md
@@ -212,3 +212,91 @@ The supported set of annotations are defined in the `NodeAnnotations.cs` source
   preventing them being merged with other issues.
 * `BuildBlocker`: Whether failures in this node should be considered a build blocker. Issues identified as build
   blockers have a special tag in Slack notifications.
+
+## Ignore Patterns
+
+Some build steps produce warnings or errors that are known, expected, or otherwise
+not actionable. Without intervention, these will create spurious Build Health issues
+and notify developers unnecessarily.
+
+__Ignore patterns__ let you suppress specific log lines so they are not ingested by the
+Build Health system. When a warning or error matches an ignore pattern, it is
+re-logged as an informational [structured log event](../Internals/StructuredLogging.md)
+with the following properties attached:
+
+| Property | Description |
+|----------|-------------|
+| `IgnoredLine` | The original text of the suppressed line |
+| `IgnorePatternFile` | Absolute path to the file that defined the matching pattern |
+| `IgnorePattern` | The regular expression that matched |
+
+Because the line is demoted to informational severity, it will not generate Build
+Health issues or trigger notifications, but it still appears in the full build log.
+The attached properties make it straightforward to trace __why__ a particular line
+was suppressed and __which pattern file__ was responsible — useful when debugging
+unexpected suppressions or reviewing which patterns are active.
+
+### File Format
+
+An ignore patterns file contains one regular expression per line. Lines starting with
+`#` are treated as comments and blank lines are ignored.
+
+```
+# IgnorePatterns.txt - suppress known warnings for this project
+# ──────────────────────────────────────────────────────────────
+
+# Third-party SDK deprecation warnings we can't fix until the next SDK update
+warning C4996: 'ThirdPartyLib::DeprecatedFunc'.*was declared deprecated
+
+# Intermittent linker note on Mac that is not a real error
+ld: warning: object file.*was built for newer macOS version
+
+# Shader compiler warnings for experimental feature behind a feature flag
+LogShaderCompiler: Warning:.*EXPERIMENTAL_FEATURE
+```
+
+Each line is compiled as a .NET regular expression, so standard regex syntax
+applies. A 200ms timeout is enforced per match to guard against catastrophic
+backtracking, so prefer simple, specific patterns.
+
+### File Discovery
+
+Horde agents automatically scan the workspace for files named
+`Build/Horde/IgnorePatterns.txt` in standard locations:
+
+| Location | Example |
+|----------|---------|
+| Engine | `Engine/Build/Horde/IgnorePatterns.txt` |
+| Top-level project directories | `MyProject/Build/Horde/IgnorePatterns.txt` |
+| Platform subdirectories | `Engine/Platforms/Linux/Build/Horde/IgnorePatterns.txt` |
+
+All discovered files are merged together; patterns from every file are active for
+the job.
+
+### Additional Pattern Files via Configuration
+
+You can reference extra pattern files from your project or stream configuration
+using the `IgnorePatternFiles` property. Paths are relative to the workspace root.
+
+In a __project__ config (applies to all streams in the project):
+
+```json
+{
+    "ignorePatternFiles": [
+        "MyProject/Config/Horde/SuppressedWarnings.txt"
+    ]
+}
+```
+
+In a __stream__ config (appended to any project-level entries):
+
+```json
+{
+    "ignorePatternFiles": [
+        "MyProject/Config/Horde/StreamSpecificIgnores.txt"
+    ]
+}
+```
+
+See the [ProjectConfig](Schema/Projects.md) and [StreamConfig](Schema/Streams.md)
+references for the full schema.
