Horde: Resolve diamond include dependencies in config system

Detect and deduplicate diamond includes in the config preprocessor.
When file A includes B and C and both B and C include D, D is now
merged only once per [ConfigIncludeRoot] scope. A warning is emitted
identifying both the first includer and the duplicate requester.

Add ResolveWithOverridesAsync to IConfigService so callers can get
the fully resolved GlobalConfig (with plugins) from a set of file
overrides. Used by the pool migration diff endpoint to compare
proposed pool configs against the current database.

Propagate config warnings through preflight validation, pool-diff,
and their corresponding dashboard UIs.

diff --git a/Engine/Source/Programs/Horde/HordeDashboard/src/backend/Api.ts b/Engine/Source/Programs/Horde/HordeDashboard/src/backend/Api.ts
--- a/Engine/Source/Programs/Horde/HordeDashboard/src/backend/Api.ts
+++ b/Engine/Source/Programs/Horde/HordeDashboard/src/backend/Api.ts
@@ -3707,6 +3707,9 @@ export type PreflightConfigResponse = {
 
 	/** Detailed response */
 	detail?: string;
+
+	/** Warnings generated during config processing (e.g. diamond include dependencies) */
+	warnings?: string[];
 }
 
 
diff --git a/Engine/Source/Programs/Horde/HordeDashboard/src/components/agents/PoolMigration.tsx b/Engine/Source/Programs/Horde/HordeDashboard/src/components/agents/PoolMigration.tsx
--- a/Engine/Source/Programs/Horde/HordeDashboard/src/components/agents/PoolMigration.tsx
+++ b/Engine/Source/Programs/Horde/HordeDashboard/src/components/agents/PoolMigration.tsx
@@ -65,6 +65,7 @@ interface PoolDiffSummary {
 
 interface PoolDiffResponse {
     validationError?: string;
+    warnings?: string[];
     files: { depotPath: string; configUri: string }[];
     unreadFiles?: string[];
     currentConfigPoolCount: number;
@@ -629,6 +630,13 @@ const PoolDiffModal: React.FC<{ diff: PoolDiffResponse; show: boolean; onClose:
                 {diff.validationError}
             </MessageBar>}
 
+            {!!diff.warnings?.length && <MessageBar messageBarType={MessageBarType.warning} style={{ marginBottom: 12 }}>
+                <Text style={{ fontFamily: "Horde Open Sans SemiBold" }}>Warnings</Text>
+                {diff.warnings.map((warning, index) => (
+                    <Text key={index} block style={{ whiteSpace: "pre-wrap" }}>{warning}</Text>
+                ))}
+            </MessageBar>}
+
             {markdownContent && <div style={{ overflowY: "auto", maxHeight: "calc(85vh - 120px)" }}>
                 <Markdown>{markdownContent}</Markdown>
             </div>}
diff --git a/Engine/Source/Programs/Horde/HordeDashboard/src/components/preflights/PreflightConfigCheck.tsx b/Engine/Source/Programs/Horde/HordeDashboard/src/components/preflights/PreflightConfigCheck.tsx
--- a/Engine/Source/Programs/Horde/HordeDashboard/src/components/preflights/PreflightConfigCheck.tsx
+++ b/Engine/Source/Programs/Horde/HordeDashboard/src/components/preflights/PreflightConfigCheck.tsx
@@ -50,7 +50,7 @@ const PreflightConfigPanel: React.FC = () => {
    const queryCluster = search.get("cluster") ? search.get("cluster")! : undefined;
 
    const navigate = useNavigate();
-   const [state, setState] = useState<{ initialCL?: string, submitting?: boolean, success?: boolean, message?: string, cluster?: string }>({ initialCL: shelvedChange, cluster: queryCluster === undefined ? "Default" : queryCluster });
+   const [state, setState] = useState<{ initialCL?: string, submitting?: boolean, success?: boolean, message?: string, warnings?: string[], cluster?: string }>({ initialCL: shelvedChange, cluster: queryCluster === undefined ? "Default" : queryCluster });
    const [servers, setServers] = useState(handler.options);
 
    if (servers === undefined) {
@@ -79,7 +79,7 @@ const PreflightConfigPanel: React.FC = () => {
 
       try {
          const response = await backend.checkPreflightConfig(parseInt(preflightCL), state.cluster === "Default" ? undefined : state.cluster);
-         setState({ submitting: false, success: response.result, message: response.message });
+         setState({ submitting: false, success: response.result, message: response.message, warnings: response.warnings });
       } catch (error) {
          setState({ submitting: false, success: false, message: error as string });
       }
@@ -147,6 +147,19 @@ const PreflightConfigPanel: React.FC = () => {
                      <Text style={{ whiteSpace: "pre-wrap", height: "472px", overflowY: "auto", padding: 18 }}> {state.message}</Text>
                   </Stack>
                </Stack>}
+               {!!state.warnings?.length && (() => {
+                  const warningColor = dashboard.getStatusColors().get(StatusColor.Warnings)!;
+                  return <Stack style={{ paddingTop: 18, paddingRight: 2 }}>
+                     <Stack style={{ paddingBottom: 12 }}>
+                        <Text styles={{ root: { fontFamily: "Horde Open Sans SemiBold", color: warningColor } }}>Warnings</Text>
+                     </Stack>
+                     <Stack style={{ border: `1px solid ${warningColor}`, backgroundColor: warningColor + "1A", padding: 18 }} tokens={{ childrenGap: 8 }}>
+                        {state.warnings.map((warning, index) => (
+                           <Text key={index} style={{ whiteSpace: "pre-wrap" }}>{warning}</Text>
+                        ))}
+                     </Stack>
+                  </Stack>;
+               })()}
             </Stack>
          </Stack>
       </Stack>
diff --git a/Engine/Source/Programs/Horde/HordeServer.Shared/Configuration/ConfigContext.cs b/Engine/Source/Programs/Horde/HordeServer.Shared/Configuration/ConfigContext.cs
--- a/Engine/Source/Programs/Horde/HordeServer.Shared/Configuration/ConfigContext.cs
+++ b/Engine/Source/Programs/Horde/HordeServer.Shared/Configuration/ConfigContext.cs
@@ -51,6 +51,27 @@ namespace HordeServer.Configuration
 		/// </summary>
 		public Dictionary<Uri, IConfigFile> Files { get; } = new Dictionary<Uri, IConfigFile>();
 
+		/// <summary>
+		/// Tracks files already processed as includes within the current include root scope,
+		/// preventing diamond dependencies from causing the same file to be merged multiple times.
+		/// Maps each included URI to the URI of the file that first included it.
+		/// Cleared when entering each new [ConfigIncludeRoot] hierarchy so that different
+		/// include roots (e.g. GlobalConfig vs ProjectConfig) can independently include the same file.
+		/// </summary>
+		public Dictionary<Uri, Uri?> ProcessedIncludes { get; } = new Dictionary<Uri, Uri?>();
+
+		/// <summary>
+		/// Depth counter for nested [ConfigIncludeRoot] processing. When zero, the next
+		/// IncludeRoot entry is a new top-level scope that should clear ProcessedIncludes.
+		/// When positive, we're recursing within the same include tree (shared dedup set).
+		/// </summary>
+		public int IncludeRootDepth { get; set; }
+
+		/// <summary>
+		/// Warnings generated during config processing (e.g. diamond include dependencies)
+		/// </summary>
+		public List<string> Warnings { get; } = new List<string>();
+
 		/// <summary>
 		/// Logger for config messages
 		/// </summary>
diff --git a/Engine/Source/Programs/Horde/HordeServer.Shared/Configuration/ConfigNode.cs b/Engine/Source/Programs/Horde/HordeServer.Shared/Configuration/ConfigNode.cs
--- a/Engine/Source/Programs/Horde/HordeServer.Shared/Configuration/ConfigNode.cs
+++ b/Engine/Source/Programs/Horde/HordeServer.Shared/Configuration/ConfigNode.cs
@@ -538,19 +538,47 @@ namespace HordeServer.Configuration
 			// Before parsing properties into this object, read all the includes recursively
 			if (IncludeRoot)
 			{
-				// Find all the includes
-				List<IConfigFile> includes = new List<IConfigFile>();
-				await ParseIncludesInternalAsync(obj, includes, context, cancellationToken);
+				// Each [ConfigIncludeRoot] type gets its own scope for diamond include dedup.
+				// When IncludeRootDepth is 0, we're entering a new top-level root (e.g. GlobalConfig,
+				// ProjectConfig, StreamConfig) and should clear the dedup set. When positive, we're
+				// recursing within the same include tree and share the set.
+				bool isNewScope = (context.IncludeRootDepth == 0);
+				if (isNewScope)
+				{
+					context.ProcessedIncludes.Clear();
+				}
+				context.IncludeRootDepth++;
 
-				// Find all the files, merge them into target
-				foreach (IConfigFile include in includes)
+				try
 				{
-					context.IncludeStack.Push(include);
+					// Find all the includes
+					List<IConfigFile> includes = new List<IConfigFile>();
+					await ParseIncludesInternalAsync(obj, includes, context, cancellationToken);
+
+					// Find all the files, merge them into target. Skip files that have already been
+					// processed to handle diamond dependencies (e.g. A includes B and C, both B and C
+					// include D - D should only be merged once).
+					string currentFile = context.CurrentFile?.ToString() ?? "(unknown)";
+					foreach (IConfigFile include in includes)
+					{
+						if (context.ProcessedIncludes.TryGetValue(include.Uri, out Uri? firstIncluder))
+						{
+							context.Warnings.Add($"Diamond include: '{include.Uri}' was already included by '{firstIncluder}' (skipped duplicate from '{currentFile}'). Consider moving this include to a common parent file to avoid redundant includes.");
+							continue;
+						}
+						context.ProcessedIncludes.Add(include.Uri, context.CurrentFile);
+
+						context.IncludeStack.Push(include);
 
-					JsonObject includedJsonObject = await context.ParseFileAsync(include, cancellationToken);
-					target = (JsonObject?)await PreprocessAsync(includedJsonObject, target, context, cancellationToken);
+						JsonObject includedJsonObject = await context.ParseFileAsync(include, cancellationToken);
+						target = (JsonObject?)await PreprocessAsync(includedJsonObject, target, context, cancellationToken);
 
-					context.IncludeStack.Pop();
+						context.IncludeStack.Pop();
+					}
+				}
+				finally
+				{
+					context.IncludeRootDepth--;
 				}
 			}
 
diff --git a/Engine/Source/Programs/Horde/HordeServer.Shared/Configuration/IConfigService.cs b/Engine/Source/Programs/Horde/HordeServer.Shared/Configuration/IConfigService.cs
--- a/Engine/Source/Programs/Horde/HordeServer.Shared/Configuration/IConfigService.cs
+++ b/Engine/Source/Programs/Horde/HordeServer.Shared/Configuration/IConfigService.cs
@@ -17,7 +17,7 @@ namespace HordeServer.Configuration
 		/// <summary>
 		/// Validate a new set of config files. Parses and runs PostLoad methods on them.
 		/// </summary>
-		Task<string?> ValidateAsync(Dictionary<Uri, byte[]> files, CancellationToken cancellationToken);
+		Task<ConfigValidationResult> ValidateAsync(Dictionary<Uri, byte[]> files, CancellationToken cancellationToken);
 
 		/// <summary>
 		/// Resolve a set of config file overrides through the full pipeline (includes, macros, inheritance)
@@ -29,6 +29,27 @@ namespace HordeServer.Configuration
 		Task<ConfigResolveResult> ResolveWithOverridesAsync(Dictionary<Uri, byte[]> files, CancellationToken cancellationToken);
 	}
 
+	/// <summary>
+	/// Result from validating config files
+	/// </summary>
+	public class ConfigValidationResult
+	{
+		/// <summary>
+		/// Whether validation succeeded
+		/// </summary>
+		public bool Success => Error == null;
+
+		/// <summary>
+		/// Error message if validation failed (null on success)
+		/// </summary>
+		public string? Error { get; set; }
+
+		/// <summary>
+		/// Warnings generated during config processing (e.g. diamond include dependencies)
+		/// </summary>
+		public List<string> Warnings { get; set; } = new List<string>();
+	}
+
 	/// <summary>
 	/// Result from resolving config with overrides
 	/// </summary>
@@ -54,5 +75,10 @@ namespace HordeServer.Configuration
 		/// Config file URIs from the overrides that were not consumed by the pipeline
 		/// </summary>
 		public List<Uri> UnreadFiles { get; set; } = new List<Uri>();
+
+		/// <summary>
+		/// Warnings generated during config processing (e.g. diamond include dependencies)
+		/// </summary>
+		public List<string> Warnings { get; set; } = new List<string>();
 	}
 }
diff --git a/Engine/Source/Programs/Horde/HordeServer.Tests/Configuration/ConfigTests.cs b/Engine/Source/Programs/Horde/HordeServer.Tests/Configuration/ConfigTests.cs
--- a/Engine/Source/Programs/Horde/HordeServer.Tests/Configuration/ConfigTests.cs
+++ b/Engine/Source/Programs/Horde/HordeServer.Tests/Configuration/ConfigTests.cs
@@ -39,6 +39,21 @@ namespace HordeServer.Tests.Configuration
 			public SubObject? TestObject { get; set; }
 		}
 
+		[ConfigIncludeRoot]
+		class NestedIncludeRoot
+		{
+			public List<ConfigInclude> Include { get; set; } = new List<ConfigInclude>();
+			public List<string> Items { get; set; } = new List<string>();
+		}
+
+		[ConfigIncludeRoot]
+		class OuterWithNestedRoot
+		{
+			public List<ConfigInclude> Include { get; set; } = new List<ConfigInclude>();
+			public List<string> Items { get; set; } = new List<string>();
+			public NestedIncludeRoot? Nested { get; set; }
+		}
+
 		readonly JsonSerializerOptions _jsonOptions = new JsonSerializerOptions { DefaultIgnoreCondition = System.Text.Json.Serialization.JsonIgnoreCondition.WhenWritingDefault };
 
 		[TestMethod]
@@ -282,6 +297,120 @@ namespace HordeServer.Tests.Configuration
 			public BaseMacroScope? ChildScope { get; set; }
 		}
 
+		[TestMethod]
+		public async Task DiamondIncludeTestAsync()
+		{
+			// Tests that diamond includes (A includes B and C, both B and C include D) only
+			// merge D's content once, preventing duplicate array entries and dictionary key errors.
+			CancellationToken cancellationToken = CancellationToken.None;
+			InMemoryConfigSource source = new InMemoryConfigSource();
+
+			// memory:///shared - the shared base included by both child1 and child2
+			Uri sharedUri = new Uri("memory:///shared");
+			{
+				ConfigObject obj = new ConfigObject();
+				obj.TestList.Add("shared-item");
+				obj.TestObject = new SubObject { ValueA = "shared" };
+				source.Add(sharedUri, JsonSerializer.SerializeToUtf8Bytes(obj, _jsonOptions));
+			}
+
+			// memory:///child1 - includes shared
+			Uri child1Uri = new Uri("memory:///child1");
+			{
+				ConfigObject obj = new ConfigObject();
+				obj.Include.Add(new ConfigInclude { Path = sharedUri.ToString() });
+				obj.TestList.Add("child1-item");
+				source.Add(child1Uri, JsonSerializer.SerializeToUtf8Bytes(obj, _jsonOptions));
+			}
+
+			// memory:///child2 - also includes shared (diamond)
+			Uri child2Uri = new Uri("memory:///child2");
+			{
+				ConfigObject obj = new ConfigObject();
+				obj.Include.Add(new ConfigInclude { Path = sharedUri.ToString() });
+				obj.TestList.Add("child2-item");
+				source.Add(child2Uri, JsonSerializer.SerializeToUtf8Bytes(obj, _jsonOptions));
+			}
+
+			// memory:///root - includes child1 and child2
+			Uri rootUri = new Uri("memory:///root");
+			{
+				ConfigObject obj = new ConfigObject();
+				obj.Include.Add(new ConfigInclude { Path = child1Uri.ToString() });
+				obj.Include.Add(new ConfigInclude { Path = child2Uri.ToString() });
+				obj.TestList.Add("root-item");
+				source.Add(rootUri, JsonSerializer.SerializeToUtf8Bytes(obj, _jsonOptions));
+			}
+
+			Dictionary<string, IConfigSource> sources = new Dictionary<string, IConfigSource>();
+			sources["memory"] = source;
+
+			ConfigContext context = new ConfigContext(_jsonOptions, sources, NullLogger.Instance);
+
+			ConfigObject result = await context.ReadAsync<ConfigObject>(rootUri, cancellationToken);
+
+			// "shared-item" should only appear once despite being included via both child1 and child2
+			Assert.IsTrue(result.TestList.SequenceEqual(new[] { "shared-item", "child1-item", "child2-item", "root-item" }));
+			Assert.AreEqual("shared", result.TestObject!.ValueA);
+
+			// A warning should be generated about the diamond include, identifying both includers
+			Assert.AreEqual(1, context.Warnings.Count);
+			Assert.IsTrue(context.Warnings[0].Contains("Diamond include"));
+			Assert.IsTrue(context.Warnings[0].Contains(sharedUri.ToString()));
+			Assert.IsTrue(context.Warnings[0].Contains(child1Uri.ToString()), "Warning should identify the first includer");
+			Assert.IsTrue(context.Warnings[0].Contains(child2Uri.ToString()), "Warning should identify the duplicate includer");
+			Assert.IsTrue(context.Warnings[0].Contains("Consider moving"));
+		}
+
+		[TestMethod]
+		public async Task CrossRootIncludeTestAsync()
+		{
+			// Tests that different [ConfigIncludeRoot] types processed in the same ConfigContext
+			// can independently include the same file. The outer root and its nested root property
+			// should both get the shared file's content without interference.
+			CancellationToken cancellationToken = CancellationToken.None;
+			InMemoryConfigSource source = new InMemoryConfigSource();
+
+			// Shared file included by both outer and nested roots
+			Uri sharedUri = new Uri("memory:///shared");
+			{
+				NestedIncludeRoot obj = new NestedIncludeRoot();
+				obj.Items.Add("shared-item");
+				source.Add(sharedUri, JsonSerializer.SerializeToUtf8Bytes(obj, _jsonOptions));
+			}
+
+			// Outer config that includes shared AND has a nested property that also includes shared
+			Uri outerUri = new Uri("memory:///outer");
+			{
+				OuterWithNestedRoot obj = new OuterWithNestedRoot();
+				obj.Include.Add(new ConfigInclude { Path = sharedUri.ToString() });
+				obj.Items.Add("outer-item");
+				// The nested property is inline JSON, its Include will be processed as a separate IncludeRoot
+				obj.Nested = new NestedIncludeRoot();
+				obj.Nested.Include.Add(new ConfigInclude { Path = sharedUri.ToString() });
+				obj.Nested.Items.Add("inner-item");
+				source.Add(outerUri, JsonSerializer.SerializeToUtf8Bytes(obj, _jsonOptions));
+			}
+
+			Dictionary<string, IConfigSource> sources = new Dictionary<string, IConfigSource>();
+			sources["memory"] = source;
+
+			ConfigContext context = new ConfigContext(_jsonOptions, sources, NullLogger.Instance);
+
+			OuterWithNestedRoot result = await context.ReadAsync<OuterWithNestedRoot>(outerUri, cancellationToken);
+
+			// Outer should have shared-item merged from the include
+			Assert.IsTrue(result.Items.SequenceEqual(new[] { "shared-item", "outer-item" }));
+
+			// Nested should ALSO have shared-item despite the outer root already including it.
+			// Different include roots should have independent dedup scopes.
+			Assert.IsNotNull(result.Nested);
+			Assert.IsTrue(result.Nested!.Items.SequenceEqual(new[] { "shared-item", "inner-item" }));
+
+			// No diamond warnings should be generated since these are independent include roots
+			Assert.AreEqual(0, context.Warnings.Count);
+		}
+
 		[TestMethod]
 		public async Task MacroTestAsync()
 		{
diff --git a/Engine/Source/Programs/Horde/HordeServer/Configuration/ConfigService.cs b/Engine/Source/Programs/Horde/HordeServer/Configuration/ConfigService.cs
--- a/Engine/Source/Programs/Horde/HordeServer/Configuration/ConfigService.cs
+++ b/Engine/Source/Programs/Horde/HordeServer/Configuration/ConfigService.cs
@@ -315,7 +315,7 @@ namespace HordeServer.Configuration
 		/// <summary>
 		/// Validate a new set of config files. Parses and runs PostLoad methods on them.
 		/// </summary>
-		public async Task<string?> ValidateAsync(Dictionary<Uri, byte[]> files, CancellationToken cancellationToken)
+		public async Task<ConfigValidationResult> ValidateAsync(Dictionary<Uri, byte[]> files, CancellationToken cancellationToken)
 		{
 			(Dictionary<Uri, OverrideConfigFile> overrideFiles, ConfigContext context) = CreateOverrideContext(files);
 			try
@@ -326,15 +326,15 @@ namespace HordeServer.Configuration
 				{
 					if (!file.WasRead)
 					{
-						return $"File {file.Uri} was not read by server";
+						return new ConfigValidationResult { Error = $"File {file.Uri} was not read by server", Warnings = context.Warnings };
 					}
 				}
 
-				return null;
+				return new ConfigValidationResult { Warnings = context.Warnings };
 			}
 			catch (Exception ex)
 			{
-				return FormatConfigException(ex, context, "Unhandled exception while validating config files");
+				return new ConfigValidationResult { Error = FormatConfigException(ex, context, "Unhandled exception while validating config files"), Warnings = context.Warnings };
 			}
 		}
 
@@ -356,11 +356,11 @@ namespace HordeServer.Configuration
 					}
 				}
 
-				return new ConfigResolveResult { Config = globalConfig, Plugins = globalConfig.Plugins, UnreadFiles = unreadFiles };
+				return new ConfigResolveResult { Config = globalConfig, Plugins = globalConfig.Plugins, UnreadFiles = unreadFiles, Warnings = context.Warnings };
 			}
 			catch (Exception ex)
 			{
-				return new ConfigResolveResult { Error = FormatConfigException(ex, context, "Unhandled exception while resolving config overrides") };
+				return new ConfigResolveResult { Error = FormatConfigException(ex, context, "Unhandled exception while resolving config overrides"), Warnings = context.Warnings };
 			}
 		}
 
diff --git a/Engine/Source/Programs/Horde/Plugins/Build/HordeServer.Build/VersionControl/Perforce/ServerPerforceController.cs b/Engine/Source/Programs/Horde/Plugins/Build/HordeServer.Build/VersionControl/Perforce/ServerPerforceController.cs
--- a/Engine/Source/Programs/Horde/Plugins/Build/HordeServer.Build/VersionControl/Perforce/ServerPerforceController.cs
+++ b/Engine/Source/Programs/Horde/Plugins/Build/HordeServer.Build/VersionControl/Perforce/ServerPerforceController.cs
@@ -2,7 +2,6 @@
 
 using EpicGames.Core;
 using EpicGames.Horde.Server;
-using EpicGames.Perforce;
 using HordeServer.Configuration;
 using HordeServer.Utilities;
 using Microsoft.AspNetCore.Authorization;
@@ -19,15 +18,15 @@ namespace HordeServer.VersionControl.Perforce
 	public class ServerPerforceController : HordeControllerBase
 	{
 		readonly IConfigService _configService;
-		readonly IPerforceService _perforceService;
+		readonly IChangelistFileReader _fileReader;
 
 		/// <summary>
 		/// Constructor
 		/// </summary>
-		public ServerPerforceController(IConfigService configService, IPerforceService perforceService)
+		public ServerPerforceController(IConfigService configService, IChangelistFileReader fileReader)
 		{
 			_configService = configService;
-			_perforceService = perforceService;
+			_fileReader = fileReader;
 		}
 
 		/// <summary>
@@ -37,45 +36,14 @@ namespace HordeServer.VersionControl.Perforce
 		[Route("/api/v1/server/preflightconfig")]
 		public async Task<ActionResult<PreflightConfigResponse>> PreflightConfigAsync(PreflightConfigRequest request, CancellationToken cancellationToken)
 		{
-			string cluster = request.Cluster ?? "default";
-
-			IPooledPerforceConnection perforce = await _perforceService.ConnectAsync(cluster, cancellationToken: cancellationToken);
-
-			PerforceResponse<DescribeRecord> describeResponse = await perforce.TryDescribeAsync(DescribeOptions.Shelved, -1, request.ShelvedChange, cancellationToken);
-			if (!describeResponse.Succeeded)
+			Dictionary<Uri, byte[]> files;
+			try
 			{
-				return BadRequest(KnownLogEvents.Horde_InvalidPreflight, "CL {Change} does not exist.", request.ShelvedChange);
+				files = await _fileReader.ReadChangelistFilesAsync(request.ShelvedChange, request.Cluster, cancellationToken);
 			}
-
-			DescribeRecord record = describeResponse.Data;
-
-			List<string> configFiles = new List<string> { "/globals.json", "global.json", ".project.json", ".stream.json", ".dashboard.json", ".telemetry.json" };
-
-			Dictionary<Uri, byte[]> files = new Dictionary<Uri, byte[]>();
-			foreach (DescribeFileRecord fileRecord in record.Files)
+			catch (ChangelistNotFoundException)
 			{
-				if (configFiles.FirstOrDefault(config => fileRecord.DepotFile.EndsWith(config, StringComparison.OrdinalIgnoreCase)) != null)
-				{
-					PerforceResponse<PrintRecord<byte[]>> printRecordResponse = await perforce.TryPrintDataAsync($"{fileRecord.DepotFile}@={request.ShelvedChange}", cancellationToken);
-					if (!printRecordResponse.Succeeded || printRecordResponse.Data.Contents == null)
-					{
-						return BadRequest($"Unable to print contents of {fileRecord.DepotFile}@={request.ShelvedChange}");
-					}
-
-					PrintRecord<byte[]> printRecord = printRecordResponse.Data;
-
-					Uri uri;
-					try
-					{
-						uri = new Uri($"perforce://{cluster}{printRecord.DepotFile}");
-					}
-					catch (Exception ex)
-					{
-						return BadRequest($"Unable to create URI for cluster '{cluster}', file '{printRecord.DepotFile}': {ex.Message}");
-					}
-
-					files.Add(uri, printRecord.Contents);
-				}
+				return BadRequest(KnownLogEvents.Horde_InvalidPreflight, "CL {Change} does not exist.", request.ShelvedChange);
 			}
 
 			if (files.Count == 0)
@@ -83,11 +51,12 @@ namespace HordeServer.VersionControl.Perforce
 				return BadRequest(KnownLogEvents.Horde_InvalidPreflight, "No config files found in CL {Change}.", request.ShelvedChange);
 			}
 
-			string? message = await _configService.ValidateAsync(files, cancellationToken);
+			ConfigValidationResult validationResult = await _configService.ValidateAsync(files, cancellationToken);
 
 			PreflightConfigResponse response = new PreflightConfigResponse();
-			response.Result = message == null;
-			response.Message = message;
+			response.Result = validationResult.Success;
+			response.Message = validationResult.Error;
+			response.Warnings = validationResult.Warnings.ToArray();
 
 			return response;
 		}
diff --git a/Engine/Source/Programs/Horde/Plugins/Compute/HordeServer.Compute/Server/PoolsServerController.cs b/Engine/Source/Programs/Horde/Plugins/Compute/HordeServer.Compute/Server/PoolsServerController.cs
--- a/Engine/Source/Programs/Horde/Plugins/Compute/HordeServer.Compute/Server/PoolsServerController.cs
+++ b/Engine/Source/Programs/Horde/Plugins/Compute/HordeServer.Compute/Server/PoolsServerController.cs
@@ -434,6 +434,7 @@ namespace HordeServer.Server
 				{
 					Changelist = cl,
 					Files = files.Keys.Select(u => new { DepotPath = u.AbsolutePath, ConfigUri = u.ToString() }).ToList(),
+					Warnings = resolveResult.Warnings,
 					ValidationError = resolveResult.Error
 				});
 			}
@@ -541,6 +542,7 @@ namespace HordeServer.Server
 				Changelist = cl,
 				Files = files.Keys.Select(u => new { DepotPath = u.AbsolutePath, ConfigUri = u.ToString() }).ToList(),
 				UnreadFiles = resolveResult.UnreadFiles.Select(u => u.ToString()).ToList(),
+				Warnings = resolveResult.Warnings,
 				ValidationError = (string?)null,
 				CurrentConfigPoolCount = currentPools.Count,
 				ProposedConfigPoolCount = proposedPools.Count,
diff --git a/Engine/Source/Programs/Shared/EpicGames.Horde/Server/ServerMessages.cs b/Engine/Source/Programs/Shared/EpicGames.Horde/Server/ServerMessages.cs
--- a/Engine/Source/Programs/Shared/EpicGames.Horde/Server/ServerMessages.cs
+++ b/Engine/Source/Programs/Shared/EpicGames.Horde/Server/ServerMessages.cs
@@ -171,6 +171,11 @@ namespace EpicGames.Horde.Server
 		/// Detailed response
 		/// </summary>
 		public string? Detail { get; set; }
+
+		/// <summary>
+		/// Warnings generated during config processing (e.g. diamond include dependencies)
+		/// </summary>
+		public string[] Warnings { get; set; } = [];
 	}
 
 	/// <summary>
