Horde: Add include support to ComputeConfig for pool configuration
Allows pool definitions to be split across multiple included files
via the standard ConfigInclude mechanism.

diff --git a/Engine/Source/Programs/Horde/Plugins/Compute/HordeServer.Compute/ComputeConfig.cs b/Engine/Source/Programs/Horde/Plugins/Compute/HordeServer.Compute/ComputeConfig.cs
--- a/Engine/Source/Programs/Horde/Plugins/Compute/HordeServer.Compute/ComputeConfig.cs
+++ b/Engine/Source/Programs/Horde/Plugins/Compute/HordeServer.Compute/ComputeConfig.cs
@@ -29,8 +29,20 @@ namespace HordeServer
 	/// <summary>
 	/// Configuration for the compute system
 	/// </summary>
+	[ConfigIncludeRoot]
+	[ConfigMacroScope]
 	public class ComputeConfig : IPluginConfig
 	{
+		/// <summary>
+		/// Includes for other configuration files
+		/// </summary>
+		public List<ConfigInclude> Include { get; set; } = new List<ConfigInclude>();
+
+		/// <summary>
+		/// Macros within the compute config
+		/// </summary>
+		public List<ConfigMacro> Macros { get; set; } = new List<ConfigMacro>();
+
 		/// <summary>
 		/// Inherited root acl
 		/// </summary>
