Horde: Updating schema documentation

#rnx

diff --git a/Engine/Source/Programs/Horde/Docs/Config/Schema/AclActions.md b/Engine/Source/Programs/Horde/Docs/Config/Schema/AclActions.md
--- a/Engine/Source/Programs/Horde/Docs/Config/Schema/AclActions.md
+++ b/Engine/Source/Programs/Horde/Docs/Config/Schema/AclActions.md
@@ -63,47 +63,43 @@
 | `DeleteAccount` |  |
 | `ViewAccount` |  |
 
-## Artifacts
-
-| Name | Description |
-| ---- | ----------- |
-| `ReadArtifact` |  |
-| `WriteArtifact` |  |
-| `DeleteArtifact` |  |
-| `UploadArtifact` |  |
-| `DownloadArtifact` |  |
-
-## Bisect
+## Agents
 
 | Name | Description |
 | ---- | ----------- |
-| `CreateBisectTask` |  |
-| `UpdateBisectTask` |  |
-| `ViewBisectTask` |  |
+| `CreateAgent` |  |
+| `CreateWorkstationAgent` |  |
+| `UpdateAgent` |  |
+| `DeleteAgent` |  |
+| `ViewAgent` |  |
+| `ListAgents` |  |
 
-## BuildHealth
+## Compute
 
 | Name | Description |
 | ---- | ----------- |
-| `AddUpdateDeleteBuildHealthFilter` |  |
+| `AddComputeTasks` |  |
+| `GetComputeTasks` |  |
+| `UbaCacheRead` |  |
+| `UbaCacheWrite` |  |
 
-## Devices
+## Leases
 
 | Name | Description |
 | ---- | ----------- |
-| `DeviceRead` |  |
-| `DeviceWrite` |  |
+| `ViewLeases` |  |
+| `ViewLeaseTasks` |  |
 
-## Jobs
+## Logs
 
 | Name | Description |
 | ---- | ----------- |
-| `CreateJob` |  |
-| `UpdateJob` |  |
-| `DeleteJob` |  |
-| `ExecuteJob` |  |
-| `RetryJobStep` |  |
-| `ViewJob` |  |
+| `CreateLog` |  |
+| `UpdateLog` |  |
+| `ViewLog` |  |
+| `WriteLogData` |  |
+| `CreateEvent` |  |
+| `ViewEvent` |  |
 
 ## Notices
 
@@ -113,27 +109,15 @@
 | `UpdateNotice` |  |
 | `DeleteNotice` |  |
 
-## Notifications
-
-| Name | Description |
-| ---- | ----------- |
-| `CreateSubscription` |  |
-
-## Projects
-
-| Name | Description |
-| ---- | ----------- |
-| `CreateProject` |  |
-| `DeleteProject` |  |
-| `UpdateProject` |  |
-| `ViewProject` |  |
-
-## Replicators
+## Pools
 
 | Name | Description |
 | ---- | ----------- |
-| `UpdateReplicator` |  |
-| `ViewReplicator` |  |
+| `CreatePool` |  |
+| `UpdatePool` |  |
+| `DeletePool` |  |
+| `ViewPool` |  |
+| `ListPools` |  |
 
 ## ServiceAccounts
 
@@ -144,23 +128,20 @@
 | `DeleteAccount` |  |
 | `ViewAccount` |  |
 
-## Streams
+## Sessions
 
 | Name | Description |
 | ---- | ----------- |
-| `CreateStream` |  |
-| `UpdateStream` |  |
-| `DeleteStream` |  |
-| `ViewStream` |  |
-| `ViewChanges` |  |
-| `ViewTemplate` |  |
+| `CreateSession` |  |
+| `ViewSession` |  |
 
-## TestData
+## Software
 
 | Name | Description |
 | ---- | ----------- |
-| `TestHealthRead` |  |
-| `TestHealthWrite` |  |
+| `UploadSoftware` |  |
+| `DownloadSoftware` |  |
+| `DeleteSoftware` |  |
 
 ## Accounts
 
@@ -182,30 +163,6 @@
 | `ViewAgent` |  |
 | `ListAgents` |  |
 
-## Artifacts
-
-| Name | Description |
-| ---- | ----------- |
-| `ReadArtifact` |  |
-| `WriteArtifact` |  |
-| `DeleteArtifact` |  |
-| `UploadArtifact` |  |
-| `DownloadArtifact` |  |
-
-## Bisect
-
-| Name | Description |
-| ---- | ----------- |
-| `CreateBisectTask` |  |
-| `UpdateBisectTask` |  |
-| `ViewBisectTask` |  |
-
-## BuildHealth
-
-| Name | Description |
-| ---- | ----------- |
-| `AddUpdateDeleteBuildHealthFilter` |  |
-
 ## Compute
 
 | Name | Description |
@@ -215,24 +172,6 @@
 | `UbaCacheRead` |  |
 | `UbaCacheWrite` |  |
 
-## Devices
-
-| Name | Description |
-| ---- | ----------- |
-| `DeviceRead` |  |
-| `DeviceWrite` |  |
-
-## Jobs
-
-| Name | Description |
-| ---- | ----------- |
-| `CreateJob` |  |
-| `UpdateJob` |  |
-| `DeleteJob` |  |
-| `ExecuteJob` |  |
-| `RetryJobStep` |  |
-| `ViewJob` |  |
-
 ## Leases
 
 | Name | Description |
@@ -259,12 +198,6 @@
 | `UpdateNotice` |  |
 | `DeleteNotice` |  |
 
-## Notifications
-
-| Name | Description |
-| ---- | ----------- |
-| `CreateSubscription` |  |
-
 ## Pools
 
 | Name | Description |
@@ -275,22 +208,6 @@
 | `ViewPool` |  |
 | `ListPools` |  |
 
-## Projects
-
-| Name | Description |
-| ---- | ----------- |
-| `CreateProject` |  |
-| `DeleteProject` |  |
-| `UpdateProject` |  |
-| `ViewProject` |  |
-
-## Replicators
-
-| Name | Description |
-| ---- | ----------- |
-| `UpdateReplicator` |  |
-| `ViewReplicator` |  |
-
 ## ServiceAccounts
 
 | Name | Description |
@@ -315,24 +232,6 @@
 | `DownloadSoftware` |  |
 | `DeleteSoftware` |  |
 
-## Streams
-
-| Name | Description |
-| ---- | ----------- |
-| `CreateStream` |  |
-| `UpdateStream` |  |
-| `DeleteStream` |  |
-| `ViewStream` |  |
-| `ViewChanges` |  |
-| `ViewTemplate` |  |
-
-## TestData
-
-| Name | Description |
-| ---- | ----------- |
-| `TestHealthRead` |  |
-| `TestHealthWrite` |  |
-
 ## Accounts
 
 | Name | Description |
@@ -353,30 +252,6 @@
 | `ViewAgent` |  |
 | `ListAgents` |  |
 
-## Artifacts
-
-| Name | Description |
-| ---- | ----------- |
-| `ReadArtifact` |  |
-| `WriteArtifact` |  |
-| `DeleteArtifact` |  |
-| `UploadArtifact` |  |
-| `DownloadArtifact` |  |
-
-## Bisect
-
-| Name | Description |
-| ---- | ----------- |
-| `CreateBisectTask` |  |
-| `UpdateBisectTask` |  |
-| `ViewBisectTask` |  |
-
-## BuildHealth
-
-| Name | Description |
-| ---- | ----------- |
-| `AddUpdateDeleteBuildHealthFilter` |  |
-
 ## Compute
 
 | Name | Description |
@@ -386,24 +261,6 @@
 | `UbaCacheRead` |  |
 | `UbaCacheWrite` |  |
 
-## Devices
-
-| Name | Description |
-| ---- | ----------- |
-| `DeviceRead` |  |
-| `DeviceWrite` |  |
-
-## Jobs
-
-| Name | Description |
-| ---- | ----------- |
-| `CreateJob` |  |
-| `UpdateJob` |  |
-| `DeleteJob` |  |
-| `ExecuteJob` |  |
-| `RetryJobStep` |  |
-| `ViewJob` |  |
-
 ## Leases
 
 | Name | Description |
@@ -430,12 +287,6 @@
 | `UpdateNotice` |  |
 | `DeleteNotice` |  |
 
-## Notifications
-
-| Name | Description |
-| ---- | ----------- |
-| `CreateSubscription` |  |
-
 ## Pools
 
 | Name | Description |
@@ -446,22 +297,6 @@
 | `ViewPool` |  |
 | `ListPools` |  |
 
-## Projects
-
-| Name | Description |
-| ---- | ----------- |
-| `CreateProject` |  |
-| `DeleteProject` |  |
-| `UpdateProject` |  |
-| `ViewProject` |  |
-
-## Replicators
-
-| Name | Description |
-| ---- | ----------- |
-| `UpdateReplicator` |  |
-| `ViewReplicator` |  |
-
 ## ServiceAccounts
 
 | Name | Description |
@@ -486,24 +321,6 @@
 | `DownloadSoftware` |  |
 | `DeleteSoftware` |  |
 
-## Streams
-
-| Name | Description |
-| ---- | ----------- |
-| `CreateStream` |  |
-| `UpdateStream` |  |
-| `DeleteStream` |  |
-| `ViewStream` |  |
-| `ViewChanges` |  |
-| `ViewTemplate` |  |
-
-## TestData
-
-| Name | Description |
-| ---- | ----------- |
-| `TestHealthRead` |  |
-| `TestHealthWrite` |  |
-
 ## Accounts
 
 | Name | Description |
@@ -524,30 +341,6 @@
 | `ViewAgent` |  |
 | `ListAgents` |  |
 
-## Artifacts
-
-| Name | Description |
-| ---- | ----------- |
-| `ReadArtifact` |  |
-| `WriteArtifact` |  |
-| `DeleteArtifact` |  |
-| `UploadArtifact` |  |
-| `DownloadArtifact` |  |
-
-## Bisect
-
-| Name | Description |
-| ---- | ----------- |
-| `CreateBisectTask` |  |
-| `UpdateBisectTask` |  |
-| `ViewBisectTask` |  |
-
-## BuildHealth
-
-| Name | Description |
-| ---- | ----------- |
-| `AddUpdateDeleteBuildHealthFilter` |  |
-
 ## Compute
 
 | Name | Description |
@@ -557,24 +350,6 @@
 | `UbaCacheRead` |  |
 | `UbaCacheWrite` |  |
 
-## Devices
-
-| Name | Description |
-| ---- | ----------- |
-| `DeviceRead` |  |
-| `DeviceWrite` |  |
-
-## Jobs
-
-| Name | Description |
-| ---- | ----------- |
-| `CreateJob` |  |
-| `UpdateJob` |  |
-| `DeleteJob` |  |
-| `ExecuteJob` |  |
-| `RetryJobStep` |  |
-| `ViewJob` |  |
-
 ## Leases
 
 | Name | Description |
@@ -601,12 +376,6 @@
 | `UpdateNotice` |  |
 | `DeleteNotice` |  |
 
-## Notifications
-
-| Name | Description |
-| ---- | ----------- |
-| `CreateSubscription` |  |
-
 ## Pools
 
 | Name | Description |
@@ -617,22 +386,6 @@
 | `ViewPool` |  |
 | `ListPools` |  |
 
-## Projects
-
-| Name | Description |
-| ---- | ----------- |
-| `CreateProject` |  |
-| `DeleteProject` |  |
-| `UpdateProject` |  |
-| `ViewProject` |  |
-
-## Replicators
-
-| Name | Description |
-| ---- | ----------- |
-| `UpdateReplicator` |  |
-| `ViewReplicator` |  |
-
 ## Secrets
 
 | Name | Description |
@@ -663,24 +416,6 @@
 | `DownloadSoftware` |  |
 | `DeleteSoftware` |  |
 
-## Streams
-
-| Name | Description |
-| ---- | ----------- |
-| `CreateStream` |  |
-| `UpdateStream` |  |
-| `DeleteStream` |  |
-| `ViewStream` |  |
-| `ViewChanges` |  |
-| `ViewTemplate` |  |
-
-## TestData
-
-| Name | Description |
-| ---- | ----------- |
-| `TestHealthRead` |  |
-| `TestHealthWrite` |  |
-
 ## Accounts
 
 | Name | Description |
@@ -840,19 +575,6 @@
 | `DownloadSoftware` |  |
 | `DeleteSoftware` |  |
 
-## Storage
-
-| Name | Description |
-| ---- | ----------- |
-| `ReadBlobs` |  |
-| `WriteBlobs` |  |
-| `ReadAliases` |  |
-| `WriteAliases` |  |
-| `DeleteAliases` |  |
-| `ReadRefs` |  |
-| `WriteRefs` |  |
-| `DeleteRefs` |  |
-
 ## Streams
 
 | Name | Description |
@@ -1244,19 +966,18 @@
 | `ViewChanges` |  |
 | `ViewTemplate` |  |
 
-## TestData
+## Symbols
 
 | Name | Description |
 | ---- | ----------- |
-| `TestHealthRead` |  |
-| `TestHealthWrite` |  |
+| `ReadSymbols` |  |
 
-## Tools
+## TestData
 
 | Name | Description |
 | ---- | ----------- |
-| `DownloadTool` |  |
-| `UploadTool` |  |
+| `TestHealthRead` |  |
+| `TestHealthWrite` |  |
 
 ## Accounts
 
diff --git a/Engine/Source/Programs/Horde/Docs/Config/Schema/Globals.md b/Engine/Source/Programs/Horde/Docs/Config/Schema/Globals.md
--- a/Engine/Source/Programs/Horde/Docs/Config/Schema/Globals.md
+++ b/Engine/Source/Programs/Horde/Docs/Config/Schema/Globals.md
@@ -56,14 +56,14 @@
 Name | Description
 ---- | -----------
 `analytics` | [AnalyticsConfig](#analyticsconfig)<br>Configuration for the analytics plugin
-`build` | [BuildConfig](#buildconfig)<br>Configuration for the build plugin
 `compute` | [ComputeConfig](#computeconfig)<br>Configuration for the compute plugin
+`exampleplugin` | [ExamplePluginConfig](#examplepluginconfig)<br>Configuration for the exampleplugin plugin
 `experimental` | [ExperimentalConfig](#experimentalconfig)<br>Configuration for the experimental plugin
 `secrets` | [SecretsConfig](#secretsconfig)<br>Configuration for the secrets plugin
+`build` | [BuildConfig](#buildconfig)<br>Configuration for the build plugin
 `storage` | [StorageConfig](#storageconfig)<br>Configuration for the storage plugin
-`exampleplugin` | [ExamplePluginConfig](#examplepluginconfig)<br>Configuration for the exampleplugin plugin
-`tools` | [ToolsConfig](#toolsconfig)<br>Configuration for the tools plugin
 `symbols` | [SymbolsConfig](#symbolsconfig)<br>Configuration for the symbols plugin
+`tools` | [ToolsConfig](#toolsconfig)<br>Configuration for the tools plugin
 
 ## AnalyticsConfig
 
@@ -71,121 +71,20 @@
 ---- | -----------
 `stores` | [TelemetryStoreConfig](Telemetry.md)`[]`<br>
 
-## BuildConfig
+## ComputeConfig
 
 Name | Description
 ---- | -----------
-`perforceClusters` | [PerforceCluster](#perforcecluster)`[]`<br>
-`devices` | [DeviceConfig](#deviceconfig)<br>
-`maxConformCount` | `integer`<br>
-`agentShutdownIfDisabledGracePeriod` | `string`<br>
-`artifactTypes` | [ArtifactTypeConfig](#artifacttypeconfig)`[]`<br>
-`projects` | [ProjectConfig](Projects.md)`[]`<br>
-`enableConformTasks` | `boolean`<br>
-`enableJobTaskConformCheck` | `boolean`<br>
-`issueFixedTag` | `string`<br>
-`issues` | [IssueConfig](Issues.md)<br>
-`enableIssueRootCauseAnalysisNotifications` | `boolean`<br>
-`autoRcaUseMultiStage` | `boolean`<br>
-`maxConcurrentRcaRequests` | `integer`<br>
-`issueFindSuspectsMaxSuccessDays` | `integer`<br>
-`issueFindSuspectsMaxScoringSeconds` | `integer`<br>
-`issueFindSuspectsMaxChanges` | `integer`<br>
-`issueFindSuspectsEnableParallelScoring` | `boolean`<br>
-`issueFindSuspectsMaxParallelism` | `integer`<br>
-`issueFindSuspectsEnableEarlyExit` | `boolean`<br>
-`issueFindSuspectsEnableBidirectionalQuery` | `boolean`<br>
-`issueFindSuspectsEnableFingerprintFallback` | `boolean`<br>
-`issueFingerprintFallbackMinScore` | `integer`<br>
-`issueFingerprintFallbackMaxAgeDays` | `integer`<br>
-`issueFingerprintFallbackExcludedTypes` | `string[]`<br>
-`enableExpandedIgnorePatterns` | `boolean`<br>
-
-## PerforceCluster
-
-Name | Description
----- | -----------
-`name` | `string`<br>
-`serviceAccount` | `string`<br>
-`canImpersonate` | `boolean`<br>
-`supportsPartitionedWorkspaces` | `boolean`<br>
-`servers` | [PerforceServer](#perforceserver)`[]`<br>
-`credentials` | [PerforceCredentials](#perforcecredentials)`[]`<br>
-`autoSdk` | [AutoSdkWorkspace](#autosdkworkspace)`[]`<br>
-
-## PerforceServer
-
-Name | Description
----- | -----------
-`serverAndPort` | `string`<br>
-`healthCheck` | `boolean`<br>
-`resolveDns` | `boolean`<br>
-`maxConformCount` | `integer`<br>
-`condition` | `string`<br>
-`properties` | `string[]`<br>
-
-## PerforceCredentials
-
-Name | Description
----- | -----------
-`userName` | `string`<br>
-`password` | `string`<br>
-`ticket` | `string`<br>
-
-## AutoSdkWorkspace
-
-Name | Description
----- | -----------
-`name` | `string`<br>
-`properties` | `string[]`<br>
-`userName` | `string`<br>
-`stream` | `string`<br>
-
-## DeviceConfig
-
-Name | Description
----- | -----------
-`platforms` | [DevicePlatformConfig](#deviceplatformconfig)`[]`<br>
-`pools` | [DevicePoolConfig](#devicepoolconfig)`[]`<br>
-
-## DevicePlatformConfig
-
-Name | Description
----- | -----------
-`id` | `string`<br>
-`name` | `string`<br>
-`models` | `string[]`<br>
-`legacyNames` | `string[]`<br>
-`legacyPerfSpecHighModel` | `string`<br>
-
-## DevicePoolConfig
-
-Name | Description
----- | -----------
-`id` | `string`<br>
-`name` | `string`<br>
-`poolType` | [DevicePoolType](#devicepooltype-enum)<br>
-`projectIds` | `string[]`<br>
-
-## DevicePoolType (Enum)
-
-The type of device pool
-
-Name | Description
----- | -----------
-`Automation` | Available to CIS jobs
-`Shared` | Shared by users with remote checking and checkouts
-
-## ArtifactTypeConfig
-
-Name | Description
----- | -----------
-`name` | `string`<br>
-`type` | `string`<br>
+`include` | [ConfigInclude](#configinclude)`[]`<br>
+`macros` | [ConfigMacro](#configmacro)`[]`<br>
 `acl` | [AclConfig](#aclconfig)<br>
-`keepCount` | `integer`<br>
-`keepDays` | `integer`<br>
-`namespaceId` | `string`<br>
+`versionEnum` | [ConfigVersion](#configversion-enum)<br>
+`telemetryStoreId` | `string`<br>
+`rates` | [AgentRateConfig](#agentrateconfig)`[]`<br>
+`clusters` | [ComputeClusterConfig](#computeclusterconfig)`[]`<br>
+`pools` | [PoolConfig](#poolconfig)`[]`<br>
+`software` | [AgentSoftwareConfig](#agentsoftwareconfig)`[]`<br>
+`networks` | [NetworkConfig](#networkconfig)`[]`<br>
 
 ## AclConfig
 
@@ -220,19 +119,6 @@
 `excludeActions` | `string[]`<br>
 `extends` | `string[]`<br>
 
-## ComputeConfig
-
-Name | Description
----- | -----------
-`acl` | [AclConfig](#aclconfig)<br>
-`versionEnum` | [ConfigVersion](#configversion-enum)<br>
-`telemetryStoreId` | `string`<br>
-`rates` | [AgentRateConfig](#agentrateconfig)`[]`<br>
-`clusters` | [ComputeClusterConfig](#computeclusterconfig)`[]`<br>
-`pools` | [PoolConfig](#poolconfig)`[]`<br>
-`software` | [AgentSoftwareConfig](#agentsoftwareconfig)`[]`<br>
-`networks` | [NetworkConfig](#networkconfig)`[]`<br>
-
 ## ConfigVersion (Enum)
 
 Name | Description
@@ -385,6 +271,12 @@
 `description` | `string`<br>
 `computeId` | `string`<br>
 
+## ExamplePluginConfig
+
+Name | Description
+---- | -----------
+`greetingMessage` | `string`<br>
+
 ## ExperimentalConfig
 
 Name | Description
@@ -457,6 +349,122 @@
 ---- | -----------
 `vaultUri` | `string`<br>
 
+## BuildConfig
+
+Name | Description
+---- | -----------
+`perforceClusters` | [PerforceCluster](#perforcecluster)`[]`<br>
+`devices` | [DeviceConfig](#deviceconfig)<br>
+`maxConformCount` | `integer`<br>
+`agentShutdownIfDisabledGracePeriod` | `string`<br>
+`artifactTypes` | [ArtifactTypeConfig](#artifacttypeconfig)`[]`<br>
+`projects` | [ProjectConfig](Projects.md)`[]`<br>
+`enableConformTasks` | `boolean`<br>
+`enableJobTaskConformCheck` | `boolean`<br>
+`issueFixedTag` | `string`<br>
+`issues` | [IssueConfig](Issues.md)<br>
+`enableIssueRootCauseAnalysisNotifications` | `boolean`<br>
+`autoRcaUseMultiStage` | `boolean`<br>
+`maxConcurrentRcaRequests` | `integer`<br>
+`issueFindSuspectsMaxSuccessDays` | `integer`<br>
+`issueFindSuspectsMaxScoringSeconds` | `integer`<br>
+`issueFindSuspectsMaxChanges` | `integer`<br>
+`issueFindSuspectsEnableParallelScoring` | `boolean`<br>
+`issueFindSuspectsMaxParallelism` | `integer`<br>
+`issueFindSuspectsEnableEarlyExit` | `boolean`<br>
+`issueFindSuspectsEnableBidirectionalQuery` | `boolean`<br>
+`issueFindSuspectsEnableFingerprintFallback` | `boolean`<br>
+`issueFingerprintFallbackMinScore` | `integer`<br>
+`issueFingerprintFallbackMaxAgeDays` | `integer`<br>
+`issueFingerprintFallbackExcludedTypes` | `string[]`<br>
+`enableExpandedIgnorePatterns` | `boolean`<br>
+
+## PerforceCluster
+
+Name | Description
+---- | -----------
+`name` | `string`<br>
+`serviceAccount` | `string`<br>
+`canImpersonate` | `boolean`<br>
+`supportsPartitionedWorkspaces` | `boolean`<br>
+`servers` | [PerforceServer](#perforceserver)`[]`<br>
+`credentials` | [PerforceCredentials](#perforcecredentials)`[]`<br>
+`autoSdk` | [AutoSdkWorkspace](#autosdkworkspace)`[]`<br>
+
+## PerforceServer
+
+Name | Description
+---- | -----------
+`serverAndPort` | `string`<br>
+`healthCheck` | `boolean`<br>
+`resolveDns` | `boolean`<br>
+`maxConformCount` | `integer`<br>
+`condition` | `string`<br>
+`properties` | `string[]`<br>
+
+## PerforceCredentials
+
+Name | Description
+---- | -----------
+`userName` | `string`<br>
+`password` | `string`<br>
+`ticket` | `string`<br>
+
+## AutoSdkWorkspace
+
+Name | Description
+---- | -----------
+`name` | `string`<br>
+`properties` | `string[]`<br>
+`userName` | `string`<br>
+`stream` | `string`<br>
+
+## DeviceConfig
+
+Name | Description
+---- | -----------
+`platforms` | [DevicePlatformConfig](#deviceplatformconfig)`[]`<br>
+`pools` | [DevicePoolConfig](#devicepoolconfig)`[]`<br>
+
+## DevicePlatformConfig
+
+Name | Description
+---- | -----------
+`id` | `string`<br>
+`name` | `string`<br>
+`models` | `string[]`<br>
+`legacyNames` | `string[]`<br>
+`legacyPerfSpecHighModel` | `string`<br>
+
+## DevicePoolConfig
+
+Name | Description
+---- | -----------
+`id` | `string`<br>
+`name` | `string`<br>
+`poolType` | [DevicePoolType](#devicepooltype-enum)<br>
+`projectIds` | `string[]`<br>
+
+## DevicePoolType (Enum)
+
+The type of device pool
+
+Name | Description
+---- | -----------
+`Automation` | Available to CIS jobs
+`Shared` | Shared by users with remote checking and checkouts
+
+## ArtifactTypeConfig
+
+Name | Description
+---- | -----------
+`name` | `string`<br>
+`type` | `string`<br>
+`acl` | [AclConfig](#aclconfig)<br>
+`keepCount` | `integer`<br>
+`keepDays` | `integer`<br>
+`namespaceId` | `string`<br>
+
 ## StorageConfig
 
 Name | Description
@@ -519,11 +527,20 @@
 `enableAliases` | `boolean`<br>
 `acl` | [AclConfig](#aclconfig)<br>
 
-## ExamplePluginConfig
+## SymbolsConfig
 
 Name | Description
 ---- | -----------
-`greetingMessage` | `string`<br>
+`stores` | [SymbolStoreConfig](#symbolstoreconfig)`[]`<br>
 
+## SymbolStoreConfig
+
+Name | Description
+---- | -----------
+`id` | `string`<br>
+`namespaceId` | `string`<br>
+`public` | `boolean`<br>
+`acl` | [AclConfig](#aclconfig)<br>
+
 ## ToolsConfig
 
@@ -548,18 +565,3 @@
 `metadata` | `string` `->` `string`<br>
 `namespaceId` | `string`<br>
 `acl` | [AclConfig](#aclconfig)<br>
-
-## SymbolsConfig
-
-Name | Description
----- | -----------
-`stores` | [SymbolStoreConfig](#symbolstoreconfig)`[]`<br>
-
-## SymbolStoreConfig
-
-Name | Description
----- | -----------
-`id` | `string`<br>
-`namespaceId` | `string`<br>
-`public` | `boolean`<br>
-`acl` | [AclConfig](#aclconfig)<br>
diff --git a/Engine/Source/Programs/Horde/Docs/Deployment/ServerSettings.md b/Engine/Source/Programs/Horde/Docs/Deployment/ServerSettings.md
--- a/Engine/Source/Programs/Horde/Docs/Deployment/ServerSettings.md
+++ b/Engine/Source/Programs/Horde/Docs/Deployment/ServerSettings.md
@@ -126,14 +126,14 @@
 Name | Description
 ---- | -----------
 `analytics` | [AnalyticsServerConfig](#analyticsserverconfig)<br>Configuration for the analytics plugin
-`build` | [BuildServerConfig](#buildserverconfig)<br>Configuration for the build plugin
 `compute` | [ComputeServerConfig](#computeserverconfig)<br>Configuration for the compute plugin
+`exampleplugin` | [ExamplePluginServerConfig](#examplepluginserverconfig)<br>Configuration for the exampleplugin plugin
 `experimental` | [ExperimentalServerConfig](#experimentalserverconfig)<br>Configuration for the experimental plugin
 `secrets` | [SecretsServerConfig](#secretsserverconfig)<br>Configuration for the secrets plugin
+`build` | [BuildServerConfig](#buildserverconfig)<br>Configuration for the build plugin
 `storage` | [StorageServerConfig](#storageserverconfig)<br>Configuration for the storage plugin
-`exampleplugin` | [ExamplePluginServerConfig](#examplepluginserverconfig)<br>Configuration for the exampleplugin plugin
-`tools` | [ToolsServerConfig](#toolsserverconfig)<br>Configuration for the tools plugin
 `symbols` | [PluginServerConfig](#pluginserverconfig)<br>Configuration for the symbols plugin
+`tools` | [ToolsServerConfig](#toolsserverconfig)<br>Configuration for the tools plugin
 
 ## AnalyticsServerConfig
 
@@ -164,6 +164,64 @@
 `retainDays` | `number`<br>
 `enabled` | `boolean`<br>
 
+## ComputeServerConfig
+
+Name | Description
+---- | -----------
+`enableUpgradeTasks` | `boolean`<br>
+`withAws` | `boolean`<br>
+`awsRegions` | `string[]`<br>
+`awsAutoScalingQueueUrls` | `string[]`<br>
+`fleetManagerV2` | [FleetManagerType](#fleetmanagertype-enum)<br>
+`fleetManagerV2Config` | `object`<br>
+`autoEnrollAgents` | `boolean`<br>
+`defaultAgentPoolSizeStrategy` | [PoolSizeStrategy](#poolsizestrategy-enum)<br>
+`agentPoolScaleOutCooldownSeconds` | `integer`<br>
+`agentPoolScaleInCooldownSeconds` | `integer`<br>
+`computeTunnelPort` | `integer`<br>
+`computeTunnelAddress` | `string`<br>
+`enabled` | `boolean`<br>
+
+## FleetManagerType (Enum)
+
+Name | Description
+---- | -----------
+`Default` | 
+`NoOp` | 
+`Aws` | 
+`AwsReuse` | 
+`AwsRecycle` | 
+`AwsAsg` | 
+
+## PoolSizeStrategy (Enum)
+
+Name | Description
+---- | -----------
+`LeaseUtilization` | 
+`JobQueue` | 
+`NoOp` | 
+`ComputeQueueAwsMetric` | 
+`LeaseUtilizationAwsMetric` | 
+
+## ExamplePluginServerConfig
+
+Name | Description
+---- | -----------
+`enabled` | `boolean`<br>
+
+## ExperimentalServerConfig
+
+Name | Description
+---- | -----------
+`enabled` | `boolean`<br>
+
+## SecretsServerConfig
+
+Name | Description
+---- | -----------
+`withAws` | `boolean`<br>
+`enabled` | `boolean`<br>
+
 ## BuildServerConfig
 
 Name | Description
@@ -300,58 +358,6 @@
 `maxChildCount` | `integer`<br>Maximum number of children in each node
 `sliceThreshold` | `integer`<br>Threshold hash value for splitting interior nodes
 
-## ComputeServerConfig
-
-Name | Description
----- | -----------
-`enableUpgradeTasks` | `boolean`<br>
-`withAws` | `boolean`<br>
-`awsRegions` | `string[]`<br>
-`awsAutoScalingQueueUrls` | `string[]`<br>
-`fleetManagerV2` | [FleetManagerType](#fleetmanagertype-enum)<br>
-`fleetManagerV2Config` | `object`<br>
-`autoEnrollAgents` | `boolean`<br>
-`defaultAgentPoolSizeStrategy` | [PoolSizeStrategy](#poolsizestrategy-enum)<br>
-`agentPoolScaleOutCooldownSeconds` | `integer`<br>
-`agentPoolScaleInCooldownSeconds` | `integer`<br>
-`computeTunnelPort` | `integer`<br>
-`computeTunnelAddress` | `string`<br>
-`enabled` | `boolean`<br>
-
-## FleetManagerType (Enum)
-
-Name | Description
----- | -----------
-`Default` | 
-`NoOp` | 
-`Aws` | 
-`AwsReuse` | 
-`AwsRecycle` | 
-`AwsAsg` | 
-
-## PoolSizeStrategy (Enum)
-
-Name | Description
----- | -----------
-`LeaseUtilization` | 
-`JobQueue` | 
-`NoOp` | 
-`ComputeQueueAwsMetric` | 
-`LeaseUtilizationAwsMetric` | 
-
-## ExperimentalServerConfig
-
-Name | Description
----- | -----------
-`enabled` | `boolean`<br>
-
-## SecretsServerConfig
-
-Name | Description
----- | -----------
-`withAws` | `boolean`<br>
-`enabled` | `boolean`<br>
-
 ## StorageServerConfig
 
 Name | Description
@@ -403,7 +409,7 @@
 `AssumeRole` | 
 `AssumeRoleWebIdentity` | 
 
-## ExamplePluginServerConfig
+## PluginServerConfig
 
 Name | Description
 ---- | -----------
@@ -469,9 +475,3 @@
 `actions` | `string[]`<br>
 `excludeActions` | `string[]`<br>
 `extends` | `string[]`<br>
-
-## PluginServerConfig
-
-Name | Description
----- | -----------
-`enabled` | `boolean`<br>
