Horde: Broaden exception handling in StreamsController batch commit retrieval

Catch all exceptions (not just PerforceException) when retrieving batch commits so failures are logged and reported via MissingCommitsHeader regardless of the error type.

diff --git a/Engine/Source/Programs/Horde/Plugins/Build/HordeServer.Build/Streams/StreamsController.cs b/Engine/Source/Programs/Horde/Plugins/Build/HordeServer.Build/Streams/StreamsController.cs
--- a/Engine/Source/Programs/Horde/Plugins/Build/HordeServer.Build/Streams/StreamsController.cs
+++ b/Engine/Source/Programs/Horde/Plugins/Build/HordeServer.Build/Streams/StreamsController.cs
@@ -8,7 +8,6 @@ using EpicGames.Horde.Jobs.Templates;
 using EpicGames.Horde.Projects;
 using EpicGames.Horde.Streams;
 using EpicGames.Horde.Users;
-using EpicGames.Perforce;
 using HordeServer.Auditing;
 using HordeServer.Commits;
 using HordeServer.Jobs;
@@ -21,6 +20,7 @@ using HordeServer.Utilities;
 using Microsoft.AspNetCore.Authorization;
 using Microsoft.AspNetCore.Http;
 using Microsoft.AspNetCore.Mvc;
+using Microsoft.Extensions.Logging;
 using Microsoft.Extensions.Options;
 using OpenTelemetry.Trace;
 
@@ -41,13 +41,14 @@ namespace HordeServer.Streams
 		private readonly IUserCollection _userCollection;
 		private readonly ScheduleService _scheduleService;
 		private readonly Tracer _tracer;
+		private readonly ILogger _logger;
 		private readonly IOptionsSnapshot<BuildConfig> _buildConfig;
 		private readonly TimeZoneInfo _timeZone;
 
 		/// <summary>
 		/// Constructor
 		/// </summary>
-		public StreamsController(IStreamCollection streamCollection, ICommitService commitService, ITemplateCollectionInternal templateCollection, IJobStepRefCollection jobStepRefCollection, IUserCollection userCollection, ScheduleService scheduleService, Tracer tracer, IClock clock, IOptionsSnapshot<BuildConfig> buildConfig)
+		public StreamsController(IStreamCollection streamCollection, ICommitService commitService, ITemplateCollectionInternal templateCollection, IJobStepRefCollection jobStepRefCollection, IUserCollection userCollection, ScheduleService scheduleService, Tracer tracer, IClock clock, IOptionsSnapshot<BuildConfig> buildConfig, ILogger<StreamsController> logger)
 		{
 			_streamCollection = streamCollection;
 			_commitService = commitService;
@@ -56,6 +57,7 @@ namespace HordeServer.Streams
 			_userCollection = userCollection;
 			_scheduleService = scheduleService;
 			_tracer = tracer;
+			_logger = logger;
 			_buildConfig = buildConfig;
 			_timeZone = clock.TimeZone;
 		}
@@ -208,7 +210,7 @@ namespace HordeServer.Streams
 		/// <param name="filter">The filter to apply to the results</param>
 		/// <param name="cancellationToken">Cancellation token for the operation</param>
 		/// Returns 200 OK with list of filtered commits. Partial failures for batch retrieval are indicated via the <see cref="MissingCommitsHeader"/> header.
-		/// <remarks>Batch request mode executed with <paramref name="changeNumbers"/> will store any commits that were not found or that threw a <see cref="PerforceException"/> in the <see cref="MissingCommitsHeader"/> response header.</remarks>
+		/// <remarks>Batch request mode executed with <paramref name="changeNumbers"/> will store any commits that were not found or that threw an exception in the <see cref="MissingCommitsHeader"/> response header.</remarks>
 		[HttpGet]
 		[Route("/api/v1/streams/{streamId}/changes")]
 		[ProducesResponseType(typeof(List<GetCommitResponse>), 200)]
@@ -271,7 +273,7 @@ namespace HordeServer.Streams
 		/// Returns 200 OK with list of filtered commits. Partial failures for batch retrieval are indicated via the <see cref="MissingCommitsHeader"/> header.
 		/// <remarks>
 		/// This endpoint is a workaround to retrieve large sets of changes which would result in 414 reponses from <see cref="GetChangesAsync"/>.
-		/// Any commits that were not found or that threw a <see cref="PerforceException"/> are stored in the <see cref="MissingCommitsHeader"/> response header.
+		/// Any commits that were not found or that threw an exception are stored in the <see cref="MissingCommitsHeader"/> response header.
 		/// </remarks>
 		[HttpPost]
 		[Route("/api/v1/streams/{streamId}/changes")]
@@ -329,8 +331,9 @@ namespace HordeServer.Streams
 					{
 						return await _commitService.GetCollection(streamConfig).GetAsync(id, cancellationToken);
 					}
-					catch (PerforceException) // Swallow PerforceExceptions, and store failed commits.
+					catch (Exception ex) when (ex is not OperationCanceledException)
 					{
+						_logger.LogWarning(ex, "Failed to retrieve commit {CommitId} for stream {StreamId}", id, streamConfig.Id);
 						failedCommits.Add(id.ToString());
 
 						return null;
