Horde: Fix compilation warnings

diff --git a/Engine/Source/Programs/Horde/HordeServer.Tests/Configuration/ConfigTests.cs b/Engine/Source/Programs/Horde/HordeServer.Tests/Configuration/ConfigTests.cs
--- a/Engine/Source/Programs/Horde/HordeServer.Tests/Configuration/ConfigTests.cs
+++ b/Engine/Source/Programs/Horde/HordeServer.Tests/Configuration/ConfigTests.cs
@@ -355,11 +355,11 @@ namespace HordeServer.Tests.Configuration
 
 			// A warning should be generated about the diamond include, identifying both includers
 			Assert.AreEqual(1, context.Warnings.Count);
-			Assert.IsTrue(context.Warnings[0].Contains("Diamond include"));
-			Assert.IsTrue(context.Warnings[0].Contains(sharedUri.ToString()));
-			Assert.IsTrue(context.Warnings[0].Contains(child1Uri.ToString()), "Warning should identify the first includer");
-			Assert.IsTrue(context.Warnings[0].Contains(child2Uri.ToString()), "Warning should identify the duplicate includer");
-			Assert.IsTrue(context.Warnings[0].Contains("Consider moving"));
+			Assert.IsTrue(context.Warnings[0].Contains("Diamond include", StringComparison.Ordinal));
+			Assert.IsTrue(context.Warnings[0].Contains(sharedUri.ToString(), StringComparison.Ordinal));
+			Assert.IsTrue(context.Warnings[0].Contains(child1Uri.ToString(), StringComparison.Ordinal), "Warning should identify the first includer");
+			Assert.IsTrue(context.Warnings[0].Contains(child2Uri.ToString(), StringComparison.Ordinal), "Warning should identify the duplicate includer");
+			Assert.IsTrue(context.Warnings[0].Contains("Consider moving", StringComparison.Ordinal));
 		}
 
 		[TestMethod]
