Horde: Swap unsupportedOperation error to Unsupported to match AWS Error codes
#rnx

diff --git a/Engine/Source/Programs/Horde/Plugins/Build/HordeServer.Build.Tests/Fleet/AwsRecyclingFleetManagerTest.cs b/Engine/Source/Programs/Horde/Plugins/Build/HordeServer.Build.Tests/Fleet/AwsRecyclingFleetManagerTest.cs
--- a/Engine/Source/Programs/Horde/Plugins/Build/HordeServer.Build.Tests/Fleet/AwsRecyclingFleetManagerTest.cs
+++ b/Engine/Source/Programs/Horde/Plugins/Build/HordeServer.Build.Tests/Fleet/AwsRecyclingFleetManagerTest.cs
@@ -176,13 +176,13 @@ namespace HordeServer.Tests.Fleet
 		}
 
 		[TestMethod]
-		public async Task ExpandOneAzUnsupportedOperationWithFallbackAsync()
+		public async Task ExpandOneAzUnsupportedWithFallbackAsync()
 		{
 			FakeAmazonEc2 ec2 = new();
 			Instance i1 = ec2.AddInstance(FakeAmazonEc2.StateStopped, InstanceType.M5Large, FakeAmazonEc2.AzUsEast1A, tagPoolName: PoolName, tagPoolId: PoolId);
 			Instance i2 = ec2.AddInstance(FakeAmazonEc2.StateStopped, InstanceType.M5Large, FakeAmazonEc2.AzUsEast1A, tagPoolName: PoolName, tagPoolId: PoolId);
  
-			ec2.SetUnsupportedOperation(FakeAmazonEc2.AzUsEast1A, InstanceType.M5Large.Value);
+			ec2.SetUnsupported(FakeAmazonEc2.AzUsEast1A, InstanceType.M5Large.Value);
 			ec2.SetCapacity(FakeAmazonEc2.AzUsEast1A, InstanceType.C5Large, 2);
  
 			await ExpandPoolAsync(ec2.Get(), 2, new(new List<string> { InstanceType.M5Large, InstanceType.C5Large }));
@@ -194,11 +194,11 @@ namespace HordeServer.Tests.Fleet
 		}
  
 		[TestMethod]
-		public async Task ExpandUnsupportedOperationNoFallbackAsync()
+		public async Task ExpandUnsupportedNoFallbackAsync()
 		{
 			FakeAmazonEc2 ec2 = new();
 			ec2.AddInstance(FakeAmazonEc2.StateStopped, InstanceType.M5Large, FakeAmazonEc2.AzUsEast1A, tagPoolName: PoolName);
-			ec2.SetUnsupportedOperation(FakeAmazonEc2.AzUsEast1A, InstanceType.M5Large.Value);
+			ec2.SetUnsupported(FakeAmazonEc2.AzUsEast1A, InstanceType.M5Large.Value);
  
 			await ExpandPoolAsync(ec2.Get(), 1, new());
  
@@ -213,7 +213,7 @@ namespace HordeServer.Tests.Fleet
 			ec2.AddInstance(FakeAmazonEc2.StateStopped, InstanceType.M5Large, FakeAmazonEc2.AzUsEast1A, tagPoolName: PoolName);
  
 			// First: unsupported, Second: no capacity, Third: works
-			ec2.SetUnsupportedOperation(FakeAmazonEc2.AzUsEast1A, InstanceType.M5Large.Value);
+			ec2.SetUnsupported(FakeAmazonEc2.AzUsEast1A, InstanceType.M5Large.Value);
 			ec2.SetCapacity(FakeAmazonEc2.AzUsEast1A, InstanceType.C5Large, 0);
 			ec2.SetCapacity(FakeAmazonEc2.AzUsEast1A, InstanceType.R5Large, 1);
  
diff --git a/Engine/Source/Programs/Horde/Plugins/Build/HordeServer.Build.Tests/Fleet/FakeAmazonEc2.cs b/Engine/Source/Programs/Horde/Plugins/Build/HordeServer.Build.Tests/Fleet/FakeAmazonEc2.cs
--- a/Engine/Source/Programs/Horde/Plugins/Build/HordeServer.Build.Tests/Fleet/FakeAmazonEc2.cs
+++ b/Engine/Source/Programs/Horde/Plugins/Build/HordeServer.Build.Tests/Fleet/FakeAmazonEc2.cs
@@ -43,7 +43,7 @@ public class FakeAmazonEc2
 		{AzUsEast1B, new () }
 	};
 
-	private readonly Dictionary<string, HashSet<string>> _unsupportedOperations = new()
+	private readonly Dictionary<string, HashSet<string>> _unsupported = new()
 	{
 		{AzUsEast1A, new HashSet<string>()},
 		{AzUsEast1B, new HashSet<string>()}
@@ -81,11 +81,11 @@ public class FakeAmazonEc2
 		_availabilityZoneCapacity[availabilityZone] = temp;
 	}
 
-	public void SetUnsupportedOperation(string availabilityZone, string instanceType)
+	public void SetUnsupported(string availabilityZone, string instanceType)
 	{
-		HashSet<string> unsupported = _unsupportedOperations.GetValueOrDefault(availabilityZone, new HashSet<string>());
+		HashSet<string> unsupported = _unsupported.GetValueOrDefault(availabilityZone, new HashSet<string>());
 		unsupported.Add(instanceType);
-		_unsupportedOperations[availabilityZone] = unsupported;
+		_unsupported[availabilityZone] = unsupported;
 	}
 
 	public Instance AddInstance(
@@ -189,12 +189,12 @@ public class FakeAmazonEc2
 			string az = _instances[instanceId].Placement.AvailabilityZone;
 			string instanceType = _instances[instanceId].InstanceType;
 
-			// Check for unsupported operations first
-			if (_unsupportedOperations.GetValueOrDefault(az, new HashSet<string>()).Contains(instanceType))
+			// Check for unsupported first
+			if (_unsupported.GetValueOrDefault(az, new HashSet<string>()).Contains(instanceType))
 			{
 				throw new AmazonEC2Exception("The requested configuration is currently not supported.")
 				{
-					ErrorCode = "UnsupportedOperation"
+					ErrorCode = "Unsupported"
 				};
 			}
 
diff --git a/Engine/Source/Programs/Horde/Plugins/Compute/HordeServer.Compute/Agents/Fleet/Providers/AwsRecyclingFleetManager.cs b/Engine/Source/Programs/Horde/Plugins/Compute/HordeServer.Compute/Agents/Fleet/Providers/AwsRecyclingFleetManager.cs
--- a/Engine/Source/Programs/Horde/Plugins/Compute/HordeServer.Compute/Agents/Fleet/Providers/AwsRecyclingFleetManager.cs
+++ b/Engine/Source/Programs/Horde/Plugins/Compute/HordeServer.Compute/Agents/Fleet/Providers/AwsRecyclingFleetManager.cs
@@ -269,11 +269,11 @@ public sealed class AwsRecyclingFleetManager : IFleetManager
 				_ec2StartInstanceCounter.Add(1, new[] { instanceTypeTag, azTag, statusTag });
 				_logger.LogInformation("Insufficient capacity for {InstanceType}", instanceType?.ToString());
 			}
-			catch (AmazonEC2Exception e) when (e.ErrorCode == "UnsupportedOperation")
+			catch (AmazonEC2Exception e) when (e.ErrorCode == "Unsupported")
 			{
 				KeyValuePair<string, object?> statusTag = KeyValuePair.Create("status", (object?)"errorUnsupported");
 				_ec2StartInstanceCounter.Add(1, new[] { instanceTypeTag, azTag, statusTag });
-				_logger.LogInformation("Unsupported operation for {InstanceType}", instanceType?.ToString());
+				_logger.LogInformation("Unsupported for {InstanceType}", instanceType?.ToString());
 			}
 			catch (Exception)
 			{
