Horde: Updating schema documentation

#rnx

diff --git a/Engine/Source/Programs/Horde/Docs/Config/Schema/AclActions.md b/Engine/Source/Programs/Horde/Docs/Config/Schema/AclActions.md
--- a/Engine/Source/Programs/Horde/Docs/Config/Schema/AclActions.md
+++ b/Engine/Source/Programs/Horde/Docs/Config/Schema/AclActions.md
@@ -37,6 +37,48 @@
 | `DeleteAccount` |  |
 | `ViewAccount` |  |
 
+## Artifacts
+
+| Name | Description |
+| ---- | ----------- |
+| `ReadArtifact` |  |
+| `WriteArtifact` |  |
+| `DeleteArtifact` |  |
+| `UploadArtifact` |  |
+| `DownloadArtifact` |  |
+
+## Bisect
+
+| Name | Description |
+| ---- | ----------- |
+| `CreateBisectTask` |  |
+| `UpdateBisectTask` |  |
+| `ViewBisectTask` |  |
+
+## BuildHealth
+
+| Name | Description |
+| ---- | ----------- |
+| `AddUpdateDeleteBuildHealthFilter` |  |
+
+## Devices
+
+| Name | Description |
+| ---- | ----------- |
+| `DeviceRead` |  |
+| `DeviceWrite` |  |
+
+## Jobs
+
+| Name | Description |
+| ---- | ----------- |
+| `CreateJob` |  |
+| `UpdateJob` |  |
+| `DeleteJob` |  |
+| `ExecuteJob` |  |
+| `RetryJobStep` |  |
+| `ViewJob` |  |
+
 ## Notices
 
 | Name | Description |
@@ -45,6 +87,28 @@
 | `UpdateNotice` |  |
 | `DeleteNotice` |  |
 
+## Notifications
+
+| Name | Description |
+| ---- | ----------- |
+| `CreateSubscription` |  |
+
+## Projects
+
+| Name | Description |
+| ---- | ----------- |
+| `CreateProject` |  |
+| `DeleteProject` |  |
+| `UpdateProject` |  |
+| `ViewProject` |  |
+
+## Replicators
+
+| Name | Description |
+| ---- | ----------- |
+| `UpdateReplicator` |  |
+| `ViewReplicator` |  |
+
 ## ServiceAccounts
 
 | Name | Description |
@@ -54,6 +118,24 @@
 | `DeleteAccount` |  |
 | `ViewAccount` |  |
 
+## Streams
+
+| Name | Description |
+| ---- | ----------- |
+| `CreateStream` |  |
+| `UpdateStream` |  |
+| `DeleteStream` |  |
+| `ViewStream` |  |
+| `ViewChanges` |  |
+| `ViewTemplate` |  |
+
+## TestData
+
+| Name | Description |
+| ---- | ----------- |
+| `TestHealthRead` |  |
+| `TestHealthWrite` |  |
+
 ## Accounts
 
 | Name | Description |
@@ -74,6 +156,30 @@
 | `ViewAgent` |  |
 | `ListAgents` |  |
 
+## Artifacts
+
+| Name | Description |
+| ---- | ----------- |
+| `ReadArtifact` |  |
+| `WriteArtifact` |  |
+| `DeleteArtifact` |  |
+| `UploadArtifact` |  |
+| `DownloadArtifact` |  |
+
+## Bisect
+
+| Name | Description |
+| ---- | ----------- |
+| `CreateBisectTask` |  |
+| `UpdateBisectTask` |  |
+| `ViewBisectTask` |  |
+
+## BuildHealth
+
+| Name | Description |
+| ---- | ----------- |
+| `AddUpdateDeleteBuildHealthFilter` |  |
+
 ## Compute
 
 | Name | Description |
@@ -83,6 +189,24 @@
 | `UbaCacheRead` |  |
 | `UbaCacheWrite` |  |
 
+## Devices
+
+| Name | Description |
+| ---- | ----------- |
+| `DeviceRead` |  |
+| `DeviceWrite` |  |
+
+## Jobs
+
+| Name | Description |
+| ---- | ----------- |
+| `CreateJob` |  |
+| `UpdateJob` |  |
+| `DeleteJob` |  |
+| `ExecuteJob` |  |
+| `RetryJobStep` |  |
+| `ViewJob` |  |
+
 ## Leases
 
 | Name | Description |
@@ -109,6 +233,12 @@
 | `UpdateNotice` |  |
 | `DeleteNotice` |  |
 
+## Notifications
+
+| Name | Description |
+| ---- | ----------- |
+| `CreateSubscription` |  |
+
 ## Pools
 
 | Name | Description |
@@ -119,6 +249,22 @@
 | `ViewPool` |  |
 | `ListPools` |  |
 
+## Projects
+
+| Name | Description |
+| ---- | ----------- |
+| `CreateProject` |  |
+| `DeleteProject` |  |
+| `UpdateProject` |  |
+| `ViewProject` |  |
+
+## Replicators
+
+| Name | Description |
+| ---- | ----------- |
+| `UpdateReplicator` |  |
+| `ViewReplicator` |  |
+
 ## ServiceAccounts
 
 | Name | Description |
@@ -143,6 +289,24 @@
 | `DownloadSoftware` |  |
 | `DeleteSoftware` |  |
 
+## Streams
+
+| Name | Description |
+| ---- | ----------- |
+| `CreateStream` |  |
+| `UpdateStream` |  |
+| `DeleteStream` |  |
+| `ViewStream` |  |
+| `ViewChanges` |  |
+| `ViewTemplate` |  |
+
+## TestData
+
+| Name | Description |
+| ---- | ----------- |
+| `TestHealthRead` |  |
+| `TestHealthWrite` |  |
+
 ## Accounts
 
 | Name | Description |
@@ -163,6 +327,30 @@
 | `ViewAgent` |  |
 | `ListAgents` |  |
 
+## Artifacts
+
+| Name | Description |
+| ---- | ----------- |
+| `ReadArtifact` |  |
+| `WriteArtifact` |  |
+| `DeleteArtifact` |  |
+| `UploadArtifact` |  |
+| `DownloadArtifact` |  |
+
+## Bisect
+
+| Name | Description |
+| ---- | ----------- |
+| `CreateBisectTask` |  |
+| `UpdateBisectTask` |  |
+| `ViewBisectTask` |  |
+
+## BuildHealth
+
+| Name | Description |
+| ---- | ----------- |
+| `AddUpdateDeleteBuildHealthFilter` |  |
+
 ## Compute
 
 | Name | Description |
@@ -172,6 +360,24 @@
 | `UbaCacheRead` |  |
 | `UbaCacheWrite` |  |
 
+## Devices
+
+| Name | Description |
+| ---- | ----------- |
+| `DeviceRead` |  |
+| `DeviceWrite` |  |
+
+## Jobs
+
+| Name | Description |
+| ---- | ----------- |
+| `CreateJob` |  |
+| `UpdateJob` |  |
+| `DeleteJob` |  |
+| `ExecuteJob` |  |
+| `RetryJobStep` |  |
+| `ViewJob` |  |
+
 ## Leases
 
 | Name | Description |
@@ -198,6 +404,12 @@
 | `UpdateNotice` |  |
 | `DeleteNotice` |  |
 
+## Notifications
+
+| Name | Description |
+| ---- | ----------- |
+| `CreateSubscription` |  |
+
 ## Pools
 
 | Name | Description |
@@ -208,6 +420,22 @@
 | `ViewPool` |  |
 | `ListPools` |  |
 
+## Projects
+
+| Name | Description |
+| ---- | ----------- |
+| `CreateProject` |  |
+| `DeleteProject` |  |
+| `UpdateProject` |  |
+| `ViewProject` |  |
+
+## Replicators
+
+| Name | Description |
+| ---- | ----------- |
+| `UpdateReplicator` |  |
+| `ViewReplicator` |  |
+
 ## ServiceAccounts
 
 | Name | Description |
@@ -232,6 +460,24 @@
 | `DownloadSoftware` |  |
 | `DeleteSoftware` |  |
 
+## Streams
+
+| Name | Description |
+| ---- | ----------- |
+| `CreateStream` |  |
+| `UpdateStream` |  |
+| `DeleteStream` |  |
+| `ViewStream` |  |
+| `ViewChanges` |  |
+| `ViewTemplate` |  |
+
+## TestData
+
+| Name | Description |
+| ---- | ----------- |
+| `TestHealthRead` |  |
+| `TestHealthWrite` |  |
+
 ## Accounts
 
 | Name | Description |
@@ -252,6 +498,30 @@
 | `ViewAgent` |  |
 | `ListAgents` |  |
 
+## Artifacts
+
+| Name | Description |
+| ---- | ----------- |
+| `ReadArtifact` |  |
+| `WriteArtifact` |  |
+| `DeleteArtifact` |  |
+| `UploadArtifact` |  |
+| `DownloadArtifact` |  |
+
+## Bisect
+
+| Name | Description |
+| ---- | ----------- |
+| `CreateBisectTask` |  |
+| `UpdateBisectTask` |  |
+| `ViewBisectTask` |  |
+
+## BuildHealth
+
+| Name | Description |
+| ---- | ----------- |
+| `AddUpdateDeleteBuildHealthFilter` |  |
+
 ## Compute
 
 | Name | Description |
@@ -261,6 +531,24 @@
 | `UbaCacheRead` |  |
 | `UbaCacheWrite` |  |
 
+## Devices
+
+| Name | Description |
+| ---- | ----------- |
+| `DeviceRead` |  |
+| `DeviceWrite` |  |
+
+## Jobs
+
+| Name | Description |
+| ---- | ----------- |
+| `CreateJob` |  |
+| `UpdateJob` |  |
+| `DeleteJob` |  |
+| `ExecuteJob` |  |
+| `RetryJobStep` |  |
+| `ViewJob` |  |
+
 ## Leases
 
 | Name | Description |
@@ -287,6 +575,12 @@
 | `UpdateNotice` |  |
 | `DeleteNotice` |  |
 
+## Notifications
+
+| Name | Description |
+| ---- | ----------- |
+| `CreateSubscription` |  |
+
 ## Pools
 
 | Name | Description |
@@ -297,6 +591,22 @@
 | `ViewPool` |  |
 | `ListPools` |  |
 
+## Projects
+
+| Name | Description |
+| ---- | ----------- |
+| `CreateProject` |  |
+| `DeleteProject` |  |
+| `UpdateProject` |  |
+| `ViewProject` |  |
+
+## Replicators
+
+| Name | Description |
+| ---- | ----------- |
+| `UpdateReplicator` |  |
+| `ViewReplicator` |  |
+
 ## ServiceAccounts
 
 | Name | Description |
@@ -321,6 +631,24 @@
 | `DownloadSoftware` |  |
 | `DeleteSoftware` |  |
 
+## Streams
+
+| Name | Description |
+| ---- | ----------- |
+| `CreateStream` |  |
+| `UpdateStream` |  |
+| `DeleteStream` |  |
+| `ViewStream` |  |
+| `ViewChanges` |  |
+| `ViewTemplate` |  |
+
+## TestData
+
+| Name | Description |
+| ---- | ----------- |
+| `TestHealthRead` |  |
+| `TestHealthWrite` |  |
+
 ## Accounts
 
 | Name | Description |
@@ -341,6 +669,30 @@
 | `ViewAgent` |  |
 | `ListAgents` |  |
 
+## Artifacts
+
+| Name | Description |
+| ---- | ----------- |
+| `ReadArtifact` |  |
+| `WriteArtifact` |  |
+| `DeleteArtifact` |  |
+| `UploadArtifact` |  |
+| `DownloadArtifact` |  |
+
+## Bisect
+
+| Name | Description |
+| ---- | ----------- |
+| `CreateBisectTask` |  |
+| `UpdateBisectTask` |  |
+| `ViewBisectTask` |  |
+
+## BuildHealth
+
+| Name | Description |
+| ---- | ----------- |
+| `AddUpdateDeleteBuildHealthFilter` |  |
+
 ## Compute
 
 | Name | Description |
@@ -350,6 +702,24 @@
 | `UbaCacheRead` |  |
 | `UbaCacheWrite` |  |
 
+## Devices
+
+| Name | Description |
+| ---- | ----------- |
+| `DeviceRead` |  |
+| `DeviceWrite` |  |
+
+## Jobs
+
+| Name | Description |
+| ---- | ----------- |
+| `CreateJob` |  |
+| `UpdateJob` |  |
+| `DeleteJob` |  |
+| `ExecuteJob` |  |
+| `RetryJobStep` |  |
+| `ViewJob` |  |
+
 ## Leases
 
 | Name | Description |
@@ -376,6 +746,12 @@
 | `UpdateNotice` |  |
 | `DeleteNotice` |  |
 
+## Notifications
+
+| Name | Description |
+| ---- | ----------- |
+| `CreateSubscription` |  |
+
 ## Pools
 
 | Name | Description |
@@ -386,6 +762,22 @@
 | `ViewPool` |  |
 | `ListPools` |  |
 
+## Projects
+
+| Name | Description |
+| ---- | ----------- |
+| `CreateProject` |  |
+| `DeleteProject` |  |
+| `UpdateProject` |  |
+| `ViewProject` |  |
+
+## Replicators
+
+| Name | Description |
+| ---- | ----------- |
+| `UpdateReplicator` |  |
+| `ViewReplicator` |  |
+
 ## Secrets
 
 | Name | Description |
@@ -416,6 +808,24 @@
 | `DownloadSoftware` |  |
 | `DeleteSoftware` |  |
 
+## Streams
+
+| Name | Description |
+| ---- | ----------- |
+| `CreateStream` |  |
+| `UpdateStream` |  |
+| `DeleteStream` |  |
+| `ViewStream` |  |
+| `ViewChanges` |  |
+| `ViewTemplate` |  |
+
+## TestData
+
+| Name | Description |
+| ---- | ----------- |
+| `TestHealthRead` |  |
+| `TestHealthWrite` |  |
+
 ## Accounts
 
 | Name | Description |
@@ -575,6 +985,19 @@
 | `DownloadSoftware` |  |
 | `DeleteSoftware` |  |
 
+## Storage
+
+| Name | Description |
+| ---- | ----------- |
+| `ReadBlobs` |  |
+| `WriteBlobs` |  |
+| `ReadAliases` |  |
+| `WriteAliases` |  |
+| `DeleteAliases` |  |
+| `ReadRefs` |  |
+| `WriteRefs` |  |
+| `DeleteRefs` |  |
+
 ## Streams
 
 | Name | Description |
@@ -966,18 +1389,19 @@
 | `ViewChanges` |  |
 | `ViewTemplate` |  |
 
-## Symbols
+## TestData
 
 | Name | Description |
 | ---- | ----------- |
-| `ReadSymbols` |  |
+| `TestHealthRead` |  |
+| `TestHealthWrite` |  |
 
-## TestData
+## Tools
 
 | Name | Description |
 | ---- | ----------- |
-| `TestHealthRead` |  |
-| `TestHealthWrite` |  |
+| `DownloadTool` |  |
+| `UploadTool` |  |
 
 ## Accounts
 
diff --git a/Engine/Source/Programs/Horde/Docs/Config/Schema/Globals.md b/Engine/Source/Programs/Horde/Docs/Config/Schema/Globals.md
--- a/Engine/Source/Programs/Horde/Docs/Config/Schema/Globals.md
+++ b/Engine/Source/Programs/Horde/Docs/Config/Schema/Globals.md
@@ -55,36 +55,131 @@
 
 Name | Description
 ---- | -----------
-`analytics` | [AnalyticsConfig](#analyticsconfig)<br>Configuration for the analytics plugin
+`build` | [BuildConfig](#buildconfig)<br>Configuration for the build plugin
 `compute` | [ComputeConfig](#computeconfig)<br>Configuration for the compute plugin
 `exampleplugin` | [ExamplePluginConfig](#examplepluginconfig)<br>Configuration for the exampleplugin plugin
-`experimental` | [ExperimentalConfig](#experimentalconfig)<br>Configuration for the experimental plugin
+`analytics` | [AnalyticsConfig](#analyticsconfig)<br>Configuration for the analytics plugin
 `secrets` | [SecretsConfig](#secretsconfig)<br>Configuration for the secrets plugin
-`build` | [BuildConfig](#buildconfig)<br>Configuration for the build plugin
 `storage` | [StorageConfig](#storageconfig)<br>Configuration for the storage plugin
-`symbols` | [SymbolsConfig](#symbolsconfig)<br>Configuration for the symbols plugin
+`experimental` | [ExperimentalConfig](#experimentalconfig)<br>Configuration for the experimental plugin
 `tools` | [ToolsConfig](#toolsconfig)<br>Configuration for the tools plugin
+`symbols` | [SymbolsConfig](#symbolsconfig)<br>Configuration for the symbols plugin
 
-## AnalyticsConfig
+## BuildConfig
 
 Name | Description
 ---- | -----------
-`stores` | [TelemetryStoreConfig](Telemetry.md)`[]`<br>
+`perforceClusters` | [PerforceCluster](#perforcecluster)`[]`<br>
+`devices` | [DeviceConfig](#deviceconfig)<br>
+`maxConformCount` | `integer`<br>
+`agentShutdownIfDisabledGracePeriod` | `string`<br>
+`artifactTypes` | [ArtifactTypeConfig](#artifacttypeconfig)`[]`<br>
+`projects` | [ProjectConfig](Projects.md)`[]`<br>
+`enableConformTasks` | `boolean`<br>
+`enableJobTaskConformCheck` | `boolean`<br>
+`issueFixedTag` | `string`<br>
+`issues` | [IssueConfig](Issues.md)<br>
+`enableIssueRootCauseAnalysisNotifications` | `boolean`<br>
+`autoRcaUseMultiStage` | `boolean`<br>
+`maxConcurrentRcaRequests` | `integer`<br>
+`issueFindSuspectsMaxSuccessDays` | `integer`<br>
+`issueFindSuspectsMaxScoringSeconds` | `integer`<br>
+`issueFindSuspectsMaxChanges` | `integer`<br>
+`issueFindSuspectsEnableParallelScoring` | `boolean`<br>
+`issueFindSuspectsMaxParallelism` | `integer`<br>
+`issueFindSuspectsEnableEarlyExit` | `boolean`<br>
+`issueFindSuspectsEnableBidirectionalQuery` | `boolean`<br>
+`issueFindSuspectsEnableFingerprintFallback` | `boolean`<br>
+`issueFingerprintFallbackMinScore` | `integer`<br>
+`issueFingerprintFallbackMaxAgeDays` | `integer`<br>
+`issueFingerprintFallbackExcludedTypes` | `string[]`<br>
+`enableExpandedIgnorePatterns` | `boolean`<br>
 
-## ComputeConfig
+## PerforceCluster
 
 Name | Description
 ---- | -----------
-`include` | [ConfigInclude](#configinclude)`[]`<br>
-`macros` | [ConfigMacro](#configmacro)`[]`<br>
+`name` | `string`<br>
+`serviceAccount` | `string`<br>
+`canImpersonate` | `boolean`<br>
+`supportsPartitionedWorkspaces` | `boolean`<br>
+`servers` | [PerforceServer](#perforceserver)`[]`<br>
+`credentials` | [PerforceCredentials](#perforcecredentials)`[]`<br>
+`autoSdk` | [AutoSdkWorkspace](#autosdkworkspace)`[]`<br>
+
+## PerforceServer
+
+Name | Description
+---- | -----------
+`serverAndPort` | `string`<br>
+`healthCheck` | `boolean`<br>
+`resolveDns` | `boolean`<br>
+`maxConformCount` | `integer`<br>
+`condition` | `string`<br>
+`properties` | `string[]`<br>
+
+## PerforceCredentials
+
+Name | Description
+---- | -----------
+`userName` | `string`<br>
+`password` | `string`<br>
+`ticket` | `string`<br>
+
+## AutoSdkWorkspace
+
+Name | Description
+---- | -----------
+`name` | `string`<br>
+`properties` | `string[]`<br>
+`userName` | `string`<br>
+`stream` | `string`<br>
+
+## DeviceConfig
+
+Name | Description
+---- | -----------
+`platforms` | [DevicePlatformConfig](#deviceplatformconfig)`[]`<br>
+`pools` | [DevicePoolConfig](#devicepoolconfig)`[]`<br>
+
+## DevicePlatformConfig
+
+Name | Description
+---- | -----------
+`id` | `string`<br>
+`name` | `string`<br>
+`models` | `string[]`<br>
+`legacyNames` | `string[]`<br>
+`legacyPerfSpecHighModel` | `string`<br>
+
+## DevicePoolConfig
+
+Name | Description
+---- | -----------
+`id` | `string`<br>
+`name` | `string`<br>
+`poolType` | [DevicePoolType](#devicepooltype-enum)<br>
+`projectIds` | `string[]`<br>
+
+## DevicePoolType (Enum)
+
+The type of device pool
+
+Name | Description
+---- | -----------
+`Automation` | Available to CIS jobs
+`Shared` | Shared by users with remote checking and checkouts
+
+## ArtifactTypeConfig
+
+Name | Description
+---- | -----------
+`name` | `string`<br>
+`type` | `string`<br>
 `acl` | [AclConfig](#aclconfig)<br>
-`versionEnum` | [ConfigVersion](#configversion-enum)<br>
-`telemetryStoreId` | `string`<br>
-`rates` | [AgentRateConfig](#agentrateconfig)`[]`<br>
-`clusters` | [ComputeClusterConfig](#computeclusterconfig)`[]`<br>
-`pools` | [PoolConfig](#poolconfig)`[]`<br>
-`software` | [AgentSoftwareConfig](#agentsoftwareconfig)`[]`<br>
-`networks` | [NetworkConfig](#networkconfig)`[]`<br>
+`keepCount` | `integer`<br>
+`keepDays` | `integer`<br>
+`namespaceId` | `string`<br>
 
 ## AclConfig
 
@@ -119,6 +214,21 @@
 `excludeActions` | `string[]`<br>
 `extends` | `string[]`<br>
 
+## ComputeConfig
+
+Name | Description
+---- | -----------
+`include` | [ConfigInclude](#configinclude)`[]`<br>
+`macros` | [ConfigMacro](#configmacro)`[]`<br>
+`acl` | [AclConfig](#aclconfig)<br>
+`versionEnum` | [ConfigVersion](#configversion-enum)<br>
+`telemetryStoreId` | `string`<br>
+`rates` | [AgentRateConfig](#agentrateconfig)`[]`<br>
+`clusters` | [ComputeClusterConfig](#computeclusterconfig)`[]`<br>
+`pools` | [PoolConfig](#poolconfig)`[]`<br>
+`software` | [AgentSoftwareConfig](#agentsoftwareconfig)`[]`<br>
+`networks` | [NetworkConfig](#networkconfig)`[]`<br>
+
 ## ConfigVersion (Enum)
 
 Name | Description
@@ -277,11 +387,11 @@
 ---- | -----------
 `greetingMessage` | `string`<br>
 
-## ExperimentalConfig
+## AnalyticsConfig
 
 Name | Description
 ---- | -----------
-`notifications` | [NotificationConfig](Notification.md)`[]`<br>
+`stores` | [TelemetryStoreConfig](Telemetry.md)`[]`<br>
 
 ## SecretsConfig
 
@@ -349,122 +459,6 @@
 ---- | -----------
 `vaultUri` | `string`<br>
 
-## BuildConfig
-
-Name | Description
----- | -----------
-`perforceClusters` | [PerforceCluster](#perforcecluster)`[]`<br>
-`devices` | [DeviceConfig](#deviceconfig)<br>
-`maxConformCount` | `integer`<br>
-`agentShutdownIfDisabledGracePeriod` | `string`<br>
-`artifactTypes` | [ArtifactTypeConfig](#artifacttypeconfig)`[]`<br>
-`projects` | [ProjectConfig](Projects.md)`[]`<br>
-`enableConformTasks` | `boolean`<br>
-`enableJobTaskConformCheck` | `boolean`<br>
-`issueFixedTag` | `string`<br>
-`issues` | [IssueConfig](Issues.md)<br>
-`enableIssueRootCauseAnalysisNotifications` | `boolean`<br>
-`autoRcaUseMultiStage` | `boolean`<br>
-`maxConcurrentRcaRequests` | `integer`<br>
-`issueFindSuspectsMaxSuccessDays` | `integer`<br>
-`issueFindSuspectsMaxScoringSeconds` | `integer`<br>
-`issueFindSuspectsMaxChanges` | `integer`<br>
-`issueFindSuspectsEnableParallelScoring` | `boolean`<br>
-`issueFindSuspectsMaxParallelism` | `integer`<br>
-`issueFindSuspectsEnableEarlyExit` | `boolean`<br>
-`issueFindSuspectsEnableBidirectionalQuery` | `boolean`<br>
-`issueFindSuspectsEnableFingerprintFallback` | `boolean`<br>
-`issueFingerprintFallbackMinScore` | `integer`<br>
-`issueFingerprintFallbackMaxAgeDays` | `integer`<br>
-`issueFingerprintFallbackExcludedTypes` | `string[]`<br>
-`enableExpandedIgnorePatterns` | `boolean`<br>
-
-## PerforceCluster
-
-Name | Description
----- | -----------
-`name` | `string`<br>
-`serviceAccount` | `string`<br>
-`canImpersonate` | `boolean`<br>
-`supportsPartitionedWorkspaces` | `boolean`<br>
-`servers` | [PerforceServer](#perforceserver)`[]`<br>
-`credentials` | [PerforceCredentials](#perforcecredentials)`[]`<br>
-`autoSdk` | [AutoSdkWorkspace](#autosdkworkspace)`[]`<br>
-
-## PerforceServer
-
-Name | Description
----- | -----------
-`serverAndPort` | `string`<br>
-`healthCheck` | `boolean`<br>
-`resolveDns` | `boolean`<br>
-`maxConformCount` | `integer`<br>
-`condition` | `string`<br>
-`properties` | `string[]`<br>
-
-## PerforceCredentials
-
-Name | Description
----- | -----------
-`userName` | `string`<br>
-`password` | `string`<br>
-`ticket` | `string`<br>
-
-## AutoSdkWorkspace
-
-Name | Description
----- | -----------
-`name` | `string`<br>
-`properties` | `string[]`<br>
-`userName` | `string`<br>
-`stream` | `string`<br>
-
-## DeviceConfig
-
-Name | Description
----- | -----------
-`platforms` | [DevicePlatformConfig](#deviceplatformconfig)`[]`<br>
-`pools` | [DevicePoolConfig](#devicepoolconfig)`[]`<br>
-
-## DevicePlatformConfig
-
-Name | Description
----- | -----------
-`id` | `string`<br>
-`name` | `string`<br>
-`models` | `string[]`<br>
-`legacyNames` | `string[]`<br>
-`legacyPerfSpecHighModel` | `string`<br>
-
-## DevicePoolConfig
-
-Name | Description
----- | -----------
-`id` | `string`<br>
-`name` | `string`<br>
-`poolType` | [DevicePoolType](#devicepooltype-enum)<br>
-`projectIds` | `string[]`<br>
-
-## DevicePoolType (Enum)
-
-The type of device pool
-
-Name | Description
----- | -----------
-`Automation` | Available to CIS jobs
-`Shared` | Shared by users with remote checking and checkouts
-
-## ArtifactTypeConfig
-
-Name | Description
----- | -----------
-`name` | `string`<br>
-`type` | `string`<br>
-`acl` | [AclConfig](#aclconfig)<br>
-`keepCount` | `integer`<br>
-`keepDays` | `integer`<br>
-`namespaceId` | `string`<br>
-
 ## StorageConfig
 
 Name | Description
@@ -527,20 +521,11 @@
 `enableAliases` | `boolean`<br>
 `acl` | [AclConfig](#aclconfig)<br>
 
-## SymbolsConfig
+## ExperimentalConfig
 
 Name | Description
 ---- | -----------
-`stores` | [SymbolStoreConfig](#symbolstoreconfig)`[]`<br>
+`notifications` | [NotificationConfig](Notification.md)`[]`<br>
 
-## SymbolStoreConfig
-
-Name | Description
----- | -----------
-`id` | `string`<br>
-`namespaceId` | `string`<br>
-`public` | `boolean`<br>
-`acl` | [AclConfig](#aclconfig)<br>
-
 ## ToolsConfig
 
@@ -565,3 +550,18 @@
 `metadata` | `string` `->` `string`<br>
 `namespaceId` | `string`<br>
 `acl` | [AclConfig](#aclconfig)<br>
+
+## SymbolsConfig
+
+Name | Description
+---- | -----------
+`stores` | [SymbolStoreConfig](#symbolstoreconfig)`[]`<br>
+
+## SymbolStoreConfig
+
+Name | Description
+---- | -----------
+`id` | `string`<br>
+`namespaceId` | `string`<br>
+`public` | `boolean`<br>
+`acl` | [AclConfig](#aclconfig)<br>
diff --git a/Engine/Source/Programs/Horde/Docs/Deployment/ServerSettings.md b/Engine/Source/Programs/Horde/Docs/Deployment/ServerSettings.md
--- a/Engine/Source/Programs/Horde/Docs/Deployment/ServerSettings.md
+++ b/Engine/Source/Programs/Horde/Docs/Deployment/ServerSettings.md
@@ -125,102 +125,15 @@
 
 Name | Description
 ---- | -----------
-`analytics` | [AnalyticsServerConfig](#analyticsserverconfig)<br>Configuration for the analytics plugin
+`build` | [BuildServerConfig](#buildserverconfig)<br>Configuration for the build plugin
 `compute` | [ComputeServerConfig](#computeserverconfig)<br>Configuration for the compute plugin
 `exampleplugin` | [ExamplePluginServerConfig](#examplepluginserverconfig)<br>Configuration for the exampleplugin plugin
-`experimental` | [ExperimentalServerConfig](#experimentalserverconfig)<br>Configuration for the experimental plugin
+`analytics` | [AnalyticsServerConfig](#analyticsserverconfig)<br>Configuration for the analytics plugin
 `secrets` | [SecretsServerConfig](#secretsserverconfig)<br>Configuration for the secrets plugin
-`build` | [BuildServerConfig](#buildserverconfig)<br>Configuration for the build plugin
 `storage` | [StorageServerConfig](#storageserverconfig)<br>Configuration for the storage plugin
-`symbols` | [PluginServerConfig](#pluginserverconfig)<br>Configuration for the symbols plugin
+`experimental` | [ExperimentalServerConfig](#experimentalserverconfig)<br>Configuration for the experimental plugin
 `tools` | [ToolsServerConfig](#toolsserverconfig)<br>Configuration for the tools plugin
+`symbols` | [PluginServerConfig](#pluginserverconfig)<br>Configuration for the symbols plugin
 
-## AnalyticsServerConfig
-
-Name | Description
----- | -----------
-`sinks` | [TelemetrySinkConfig](#telemetrysinkconfig)<br>
-`enabled` | `boolean`<br>
-
-## TelemetrySinkConfig
-
-Name | Description
----- | -----------
-`epic` | [EpicTelemetryConfig](#epictelemetryconfig)<br>
-`mongo` | [MongoTelemetryConfig](#mongotelemetryconfig)<br>
-
-## EpicTelemetryConfig
-
-Name | Description
----- | -----------
-`url` | `string`<br>
-`appId` | `string`<br>
-`enabled` | `boolean`<br>
-
-## MongoTelemetryConfig
-
-Name | Description
----- | -----------
-`retainDays` | `number`<br>
-`enabled` | `boolean`<br>
-
-## ComputeServerConfig
-
-Name | Description
----- | -----------
-`enableUpgradeTasks` | `boolean`<br>
-`withAws` | `boolean`<br>
-`awsRegions` | `string[]`<br>
-`awsAutoScalingQueueUrls` | `string[]`<br>
-`fleetManagerV2` | [FleetManagerType](#fleetmanagertype-enum)<br>
-`fleetManagerV2Config` | `object`<br>
-`autoEnrollAgents` | `boolean`<br>
-`defaultAgentPoolSizeStrategy` | [PoolSizeStrategy](#poolsizestrategy-enum)<br>
-`agentPoolScaleOutCooldownSeconds` | `integer`<br>
-`agentPoolScaleInCooldownSeconds` | `integer`<br>
-`computeTunnelPort` | `integer`<br>
-`computeTunnelAddress` | `string`<br>
-`enabled` | `boolean`<br>
-
-## FleetManagerType (Enum)
-
-Name | Description
----- | -----------
-`Default` | 
-`NoOp` | 
-`Aws` | 
-`AwsReuse` | 
-`AwsRecycle` | 
-`AwsAsg` | 
-
-## PoolSizeStrategy (Enum)
-
-Name | Description
----- | -----------
-`LeaseUtilization` | 
-`JobQueue` | 
-`NoOp` | 
-`ComputeQueueAwsMetric` | 
-`LeaseUtilizationAwsMetric` | 
-
-## ExamplePluginServerConfig
-
-Name | Description
----- | -----------
-`enabled` | `boolean`<br>
-
-## ExperimentalServerConfig
-
-Name | Description
----- | -----------
-`enabled` | `boolean`<br>
-
-## SecretsServerConfig
-
-Name | Description
----- | -----------
-`withAws` | `boolean`<br>
-`enabled` | `boolean`<br>
-
 ## BuildServerConfig
 
@@ -358,6 +271,87 @@
 `maxChildCount` | `integer`<br>Maximum number of children in each node
 `sliceThreshold` | `integer`<br>Threshold hash value for splitting interior nodes
 
+## ComputeServerConfig
+
+Name | Description
+---- | -----------
+`enableUpgradeTasks` | `boolean`<br>
+`withAws` | `boolean`<br>
+`awsRegions` | `string[]`<br>
+`awsAutoScalingQueueUrls` | `string[]`<br>
+`fleetManagerV2` | [FleetManagerType](#fleetmanagertype-enum)<br>
+`fleetManagerV2Config` | `object`<br>
+`autoEnrollAgents` | `boolean`<br>
+`defaultAgentPoolSizeStrategy` | [PoolSizeStrategy](#poolsizestrategy-enum)<br>
+`agentPoolScaleOutCooldownSeconds` | `integer`<br>
+`agentPoolScaleInCooldownSeconds` | `integer`<br>
+`computeTunnelPort` | `integer`<br>
+`computeTunnelAddress` | `string`<br>
+`enabled` | `boolean`<br>
+
+## FleetManagerType (Enum)
+
+Name | Description
+---- | -----------
+`Default` | 
+`NoOp` | 
+`Aws` | 
+`AwsReuse` | 
+`AwsRecycle` | 
+`AwsAsg` | 
+
+## PoolSizeStrategy (Enum)
+
+Name | Description
+---- | -----------
+`LeaseUtilization` | 
+`JobQueue` | 
+`NoOp` | 
+`ComputeQueueAwsMetric` | 
+`LeaseUtilizationAwsMetric` | 
+
+## ExamplePluginServerConfig
+
+Name | Description
+---- | -----------
+`enabled` | `boolean`<br>
+
+## AnalyticsServerConfig
+
+Name | Description
+---- | -----------
+`sinks` | [TelemetrySinkConfig](#telemetrysinkconfig)<br>
+`enabled` | `boolean`<br>
+
+## TelemetrySinkConfig
+
+Name | Description
+---- | -----------
+`epic` | [EpicTelemetryConfig](#epictelemetryconfig)<br>
+`mongo` | [MongoTelemetryConfig](#mongotelemetryconfig)<br>
+
+## EpicTelemetryConfig
+
+Name | Description
+---- | -----------
+`url` | `string`<br>
+`appId` | `string`<br>
+`enabled` | `boolean`<br>
+
+## MongoTelemetryConfig
+
+Name | Description
+---- | -----------
+`retainDays` | `number`<br>
+`enabled` | `boolean`<br>
+
+## SecretsServerConfig
+
+Name | Description
+---- | -----------
+`withAws` | `boolean`<br>
+`enabled` | `boolean`<br>
+
 ## StorageServerConfig
 
 Name | Description
@@ -409,7 +403,7 @@
 `AssumeRole` | 
 `AssumeRoleWebIdentity` | 
 
-## PluginServerConfig
+## ExperimentalServerConfig
 
 Name | Description
 ---- | -----------
@@ -475,3 +469,9 @@
 `actions` | `string[]`<br>
 `excludeActions` | `string[]`<br>
 `extends` | `string[]`<br>
+
+## PluginServerConfig
+
+Name | Description
+---- | -----------
+`enabled` | `boolean`<br>
