Horde: Fix preflight config check warning overflow

diff --git a/Engine/Source/Programs/Horde/HordeDashboard/src/components/preflights/PreflightConfigCheck.tsx b/Engine/Source/Programs/Horde/HordeDashboard/src/components/preflights/PreflightConfigCheck.tsx
--- a/Engine/Source/Programs/Horde/HordeDashboard/src/components/preflights/PreflightConfigCheck.tsx
+++ b/Engine/Source/Programs/Horde/HordeDashboard/src/components/preflights/PreflightConfigCheck.tsx
@@ -91,10 +91,10 @@ const PreflightConfigPanel: React.FC = () => {
       return null;
    }
 
-   return (<Stack>
-      <Stack styles={{ root: { paddingTop: 18, paddingLeft: 0, paddingRight: 0, width: "100%" } }} >
-         <Stack tokens={{ childrenGap: 12 }} >
-            <Stack style={{ width: 800, paddingLeft: 12 }}>
+   return (<Stack style={{ flex: 1, minHeight: 0 }}>
+      <Stack styles={{ root: { paddingTop: 18, paddingLeft: 0, paddingRight: 0, width: "100%", flex: 1, minHeight: 0 } }} >
+         <Stack tokens={{ childrenGap: 12 }} style={{ flex: 1, minHeight: 0, display: "flex" }}>
+            <Stack style={{ width: 800, paddingLeft: 12, flex: 1, minHeight: 0, display: "flex", flexDirection: "column" }}>
                <Stack horizontal verticalAlign="center" tokens={{ childrenGap: 24 }}>
                   <Stack horizontal verticalAlign="center" tokens={{ childrenGap: 12 }}>
                      <Text styles={{ root: { fontFamily: "Horde Open Sans SemiBold" } }}>Shelved Change</Text>
@@ -139,27 +139,25 @@ const PreflightConfigPanel: React.FC = () => {
                      }}>Check</PrimaryButton>
                   </Stack>
                </Stack>
-               {!!state.message && <Stack style={{ paddingTop: 18, paddingRight: 2 }}>
+               {(!!state.message || !!state.warnings?.length) && <Stack style={{ paddingTop: 18, paddingRight: 2, paddingBottom: 18, flex: 1, minHeight: 0 }}>
                   <Stack style={{ paddingBottom: 12 }}>
-                     <Text styles={{ root: { fontFamily: "Horde Open Sans SemiBold" } }}>Error Message</Text>
+                     <Text styles={{ root: { fontFamily: "Horde Open Sans SemiBold" } }}>{state.success === false ? "Error Message" : "Result"}</Text>
                   </Stack>
-                  <Stack style={{ border: "1px solid #CDCBC9" }}>
-                     <Text style={{ whiteSpace: "pre-wrap", height: "472px", overflowY: "auto", padding: 18 }}> {state.message}</Text>
+                  <Stack style={{ border: "1px solid #CDCBC9", flex: 1, minHeight: 0, overflowY: "auto", padding: 18 }} tokens={{ childrenGap: 18 }}>
+                     {!!state.message && <Text style={{ whiteSpace: "pre-wrap" }}>{state.message}</Text>}
+                     {!!state.warnings?.length && (() => {
+                        const warningColor = dashboard.getStatusColors().get(StatusColor.Warnings)!;
+                        return <Stack tokens={{ childrenGap: 8 }}>
+                           <Text styles={{ root: { fontFamily: "Horde Open Sans SemiBold", color: warningColor } }}>Warnings</Text>
+                           <Stack style={{ border: `1px solid ${warningColor}`, backgroundColor: warningColor + "1A", padding: 12 }} tokens={{ childrenGap: 8 }}>
+                              {state.warnings.map((warning, index) => (
+                                 <Text key={index} style={{ whiteSpace: "pre-wrap" }}>{warning}</Text>
+                              ))}
+                           </Stack>
+                        </Stack>;
+                     })()}
                   </Stack>
                </Stack>}
-               {!!state.warnings?.length && (() => {
-                  const warningColor = dashboard.getStatusColors().get(StatusColor.Warnings)!;
-                  return <Stack style={{ paddingTop: 18, paddingRight: 2 }}>
-                     <Stack style={{ paddingBottom: 12 }}>
-                        <Text styles={{ root: { fontFamily: "Horde Open Sans SemiBold", color: warningColor } }}>Warnings</Text>
-                     </Stack>
-                     <Stack style={{ border: `1px solid ${warningColor}`, backgroundColor: warningColor + "1A", padding: 18 }} tokens={{ childrenGap: 8 }}>
-                        {state.warnings.map((warning, index) => (
-                           <Text key={index} style={{ whiteSpace: "pre-wrap" }}>{warning}</Text>
-                        ))}
-                     </Stack>
-                  </Stack>;
-               })()}
             </Stack>
          </Stack>
       </Stack>
@@ -174,8 +172,8 @@ export const PreflightConfigModal: React.FC<{ onClose: () => void }> = ({ onClos
    return <Stack>
       <Modal isOpen={true} isBlocking={true} topOffsetFixed={true} styles={{ main: { padding: 8, width: 880, height: 720, hasBeenOpened: false, top: "80px", position: "absolute" } }} onDismiss={() => onClose()} className={hordeClasses.modal}>
          <Stack className="horde-no-darktheme" styles={{ root: { paddingTop: 10, paddingRight: 0 } }}>
-            <Stack style={{ paddingLeft: 24, paddingRight: 12 }}>
-               <Stack tokens={{ childrenGap: 12 }} style={{ height: 700 }}>
+            <Stack style={{ paddingLeft: 24, paddingRight: 12, flex: 1, minHeight: 0 }}>
+               <Stack tokens={{ childrenGap: 12 }} style={{ height: 700, display: "flex", flexDirection: "column" }}>
                   <Stack horizontal verticalAlign="start">
                      <Stack style={{ paddingTop: 3 }}>
                         <Text variant="mediumPlus" styles={{ root: { fontFamily: "Horde Open Sans SemiBold" } }}>Preflight Configuration</Text>
@@ -188,7 +186,7 @@ export const PreflightConfigModal: React.FC<{ onClose: () => void }> = ({ onClos
                         />
                      </Stack>
                   </Stack>
-                  <Stack styles={{ root: { paddingLeft: 4, paddingRight: 0, paddingTop: 8, paddingBottom: 4 } }}>
+                  <Stack styles={{ root: { paddingLeft: 4, paddingRight: 0, paddingTop: 8, paddingBottom: 4, flex: 1, minHeight: 0 } }}>
                      <PreflightConfigPanel />
                   </Stack>
                </Stack>
