Horde: Ensure perforce service check

diff --git a/Engine/Source/Programs/Horde/Plugins/Build/HordeServer.Build/Jobs/JobsController.cs b/Engine/Source/Programs/Horde/Plugins/Build/HordeServer.Build/Jobs/JobsController.cs
--- a/Engine/Source/Programs/Horde/Plugins/Build/HordeServer.Build/Jobs/JobsController.cs
+++ b/Engine/Source/Programs/Horde/Plugins/Build/HordeServer.Build/Jobs/JobsController.cs
@@ -146,7 +146,7 @@ namespace HordeServer.Jobs
 			if (create.PreflightCommitId != null)
 			{
 				int? preflightChange = create.PreflightCommitId.TryGetPerforceChange();
-				if (preflightChange != null)
+				if (preflightChange != null && String.Equals(streamConfig.VCS, "Perforce", StringComparison.OrdinalIgnoreCase))
 				{
 					// P4 path: full shelf validation
 					(CheckShelfResult result, shelfInfo) = await _perforce.CheckShelfAsync(streamConfig, preflightChange.Value, HttpContext.RequestAborted);
