Horde: Add batch log events API and step history error/warning display
Add POST /api/v1/logs/events endpoint for fetching events across multiple logs in a single request.
Display error/warning events inline in step history table with expand/collapse buttons.

diff --git a/Engine/Source/Programs/Horde/HordeDashboard/src/backend/Api.ts b/Engine/Source/Programs/Horde/HordeDashboard/src/backend/Api.ts
--- a/Engine/Source/Programs/Horde/HordeDashboard/src/backend/Api.ts
+++ b/Engine/Source/Programs/Horde/HordeDashboard/src/backend/Api.ts
@@ -375,6 +375,15 @@ export type LogEventQuery = {
 	count?: number;
 }
 
+export type GetBatchLogEventsRequest = {
+	logIds: string[];
+	count?: number;
+	severity?: EventSeverity;
+}
+
+export type GetBatchLogEventsResponse = {
+	logEvents: Record<string, EventData[]>;
+}
 
 export type UsersQuery = {
 	ids?: string[];
diff --git a/Engine/Source/Programs/Horde/HordeDashboard/src/backend/Backend.ts b/Engine/Source/Programs/Horde/HordeDashboard/src/backend/Backend.ts
--- a/Engine/Source/Programs/Horde/HordeDashboard/src/backend/Backend.ts
+++ b/Engine/Source/Programs/Horde/HordeDashboard/src/backend/Backend.ts
@@ -1,7 +1,7 @@
 // Copyright Epic Games, Inc. All Rights Reserved.
 
 import templateCache from '../backend/TemplateCache';
-import { AccountClaimMessage, AddBuildHealthFiltersRequest, AgentData, AgentQuery, ApproveAgentsRequest, AuditLogEntry, AuditLogQuery, BatchUpdatePoolRequest, ChangeSummaryData, CreateAccountRequest, CreateBisectTaskRequest, CreateBisectTaskResponse, CreateDeviceRequest, CreateDeviceResponse, CreateExternalIssueRequest, CreateExternalIssueResponse, CreateJobRequest, CreateJobResponse, CreateNoticeRequest, CreatePoolRequest, CreateServiceAccountRequest, CreateServiceAccountResponse, CreateSoftwareResponse, CreateSubscriptionRequest, CreateSubscriptionResponse, CreateZipRequest, DashboardPreference, DeepLinkResponse, DevicePoolTelemetryQuery, DeviceTelemetryQuery, EventData, FindArtifactsResponse, FindIssueResponse, FindJobTimingsResponse, GetAccountResponse, GetAgentLeaseResponse, GetAgentTelemetrySampleResponse, GetArtifactDirectoryResponse, GetArtifactResponse, GetBisectTaskResponse, GetBuildHealthFilterResponse, GetBuildHealthFiltersRequest, GetDashboardConfigResponse, GetDevicePlatformResponse, GetDevicePoolResponse, GetDevicePoolTelemetryResponse, GetDeviceReservationResponse, GetDeviceResponse, GetDeviceTelemetryResponse, GetExternalIssueProjectResponse, GetExternalIssueResponse, GetIssueStreamResponse, GetJobsTabResponse, GetJobStepRefResponse, GetJobStepTraceResponse, GetJobTimingResponse, GetLogEventResponse, GetLogFileResponse, GetNoticeResponse, GetNotificationResponse, GetPendingAgentsResponse, GetPerforceServerStatusResponse, GetPoolResponse, GetPoolSummaryResponse, GetServerInfoResponse, GetServerSettingsResponse, GetServiceAccountResponse, GetSoftwareResponse, GetSubscriptionResponse, GetTestDataDetailsResponse, GetTestDataRefResponse, GetTestMetaResponse, GetTestResponse, GetTestsRequest, GetTestStreamResponse, GetToolResponse, GetToolSummaryResponse, GetUserResponse, GetUtilizationTelemetryResponse, GlobalConfig, IssueData, IssueQuery, IssueQueryV2, JobData, JobQuery, JobsTabColumnType, JobStepOutcome, JobStreamQuery, JobTimingsQuery, LeaseData, LogData, LogEventQuery, LogLineData, PoolData, PoolQuery, PreflightConfigResponse, ProjectData, ScheduleData, ScheduleQuery, SearchLogFileResponse, ServerPluginInfoResponse, ServerStatusResponse, ServerUpdateResponse, SessionData, StreamData, TabType, TestData, UpdateAccountRequest, UpdateAgentRequest, UpdateBisectTaskRequest, UpdateBuildHealthFiltersRequest, UpdateCurrentAccountRequest, UpdateDeploymentRequest, UpdateDeviceRequest, UpdateGlobalConfigRequest, UpdateIssueRequest, UpdateJobRequest, UpdateLeaseRequest, UpdateNoticeRequest, UpdateNotificationsRequest, UpdatePoolRequest, UpdateServerSettingsRequest, UpdateServiceAccountRequest, UpdateServiceAccountResponse, UpdateStepRequest, UpdateStepResponse, UpdateTemplateRefRequest, UpdateUserRequest, UsersQuery } from './Api';
+import { AccountClaimMessage, AddBuildHealthFiltersRequest, AgentData, AgentQuery, ApproveAgentsRequest, AuditLogEntry, AuditLogQuery, BatchUpdatePoolRequest, ChangeSummaryData, CreateAccountRequest, CreateBisectTaskRequest, CreateBisectTaskResponse, CreateDeviceRequest, CreateDeviceResponse, CreateExternalIssueRequest, CreateExternalIssueResponse, CreateJobRequest, CreateJobResponse, CreateNoticeRequest, CreatePoolRequest, CreateServiceAccountRequest, CreateServiceAccountResponse, CreateSoftwareResponse, CreateSubscriptionRequest, CreateSubscriptionResponse, CreateZipRequest, DashboardPreference, DeepLinkResponse, DevicePoolTelemetryQuery, DeviceTelemetryQuery, EventData, FindArtifactsResponse, FindIssueResponse, FindJobTimingsResponse, GetAccountResponse, GetAgentLeaseResponse, GetAgentTelemetrySampleResponse, GetArtifactDirectoryResponse, GetArtifactResponse, GetBatchLogEventsRequest, GetBatchLogEventsResponse, GetBisectTaskResponse, GetBuildHealthFilterResponse, GetBuildHealthFiltersRequest, GetDashboardConfigResponse, GetDevicePlatformResponse, GetDevicePoolResponse, GetDevicePoolTelemetryResponse, GetDeviceReservationResponse, GetDeviceResponse, GetDeviceTelemetryResponse, GetExternalIssueProjectResponse, GetExternalIssueResponse, GetIssueStreamResponse, GetJobsTabResponse, GetJobStepRefResponse, GetJobStepTraceResponse, GetJobTimingResponse, GetLogEventResponse, GetLogFileResponse, GetNoticeResponse, GetNotificationResponse, GetPendingAgentsResponse, GetPerforceServerStatusResponse, GetPoolResponse, GetPoolSummaryResponse, GetServerInfoResponse, GetServerSettingsResponse, GetServiceAccountResponse, GetSoftwareResponse, GetSubscriptionResponse, GetTestDataDetailsResponse, GetTestDataRefResponse, GetTestMetaResponse, GetTestResponse, GetTestsRequest, GetTestStreamResponse, GetToolResponse, GetToolSummaryResponse, GetUserResponse, GetUtilizationTelemetryResponse, GlobalConfig, IssueData, IssueQuery, IssueQueryV2, JobData, JobQuery, JobsTabColumnType, JobStepOutcome, JobStreamQuery, JobTimingsQuery, LeaseData, LogData, LogEventQuery, LogLineData, PoolData, PoolQuery, PreflightConfigResponse, ProjectData, ScheduleData, ScheduleQuery, SearchLogFileResponse, ServerPluginInfoResponse, ServerStatusResponse, ServerUpdateResponse, SessionData, StreamData, TabType, TestData, UpdateAccountRequest, UpdateAgentRequest, UpdateBisectTaskRequest, UpdateBuildHealthFiltersRequest, UpdateCurrentAccountRequest, UpdateDeploymentRequest, UpdateDeviceRequest, UpdateGlobalConfigRequest, UpdateIssueRequest, UpdateJobRequest, UpdateLeaseRequest, UpdateNoticeRequest, UpdateNotificationsRequest, UpdatePoolRequest, UpdateServerSettingsRequest, UpdateServiceAccountRequest, UpdateServiceAccountResponse, UpdateStepRequest, UpdateStepResponse, UpdateTemplateRefRequest, UpdateUserRequest, UsersQuery } from './Api';
 import dashboard, { Dashboard } from './Dashboard';
 import { ChallengeStatus, Fetch } from './Fetch';
 import { projectStore } from './ProjectStore';
@@ -448,6 +448,16 @@ export class Backend {
         });
     }
 
+    getBatchLogEvents(request: GetBatchLogEventsRequest): Promise<GetBatchLogEventsResponse> {
+        return new Promise<GetBatchLogEventsResponse>((resolve, reject) => {
+            this._fetch.post(`/api/v1/logs/events`, request).then((value) => {
+                resolve(value.data as GetBatchLogEventsResponse);
+            }).catch(reason => {
+                reject(reason);
+            });
+        });
+    }
+
     searchLog(logId: string, text: string, firstLine?: number, count?: number) {
 
         return new Promise<SearchLogFileResponse>((resolve, reject) => {
diff --git a/Engine/Source/Programs/Horde/HordeDashboard/src/components/jobDetailsV2/JobDetailStepHistory.tsx b/Engine/Source/Programs/Horde/HordeDashboard/src/components/jobDetailsV2/JobDetailStepHistory.tsx
--- a/Engine/Source/Programs/Horde/HordeDashboard/src/components/jobDetailsV2/JobDetailStepHistory.tsx
+++ b/Engine/Source/Programs/Horde/HordeDashboard/src/components/jobDetailsV2/JobDetailStepHistory.tsx
@@ -1,16 +1,17 @@
 // Copyright Epic Games, Inc. All Rights Reserved.
 
-import { DetailsList, DetailsListLayoutMode, DetailsRow, FontIcon, IColumn, IDetailsListProps, SelectionMode, Stack, Text } from "@fluentui/react";
+import { DefaultButton, DetailsList, DetailsListLayoutMode, DetailsRow, FontIcon, IColumn, IDetailsListProps, SelectionMode, Spinner, SpinnerSize, Stack, Text } from "@fluentui/react";
 import { observer } from "mobx-react-lite";
 import moment from "moment";
-import { useEffect } from "react";
+import { useEffect, useState } from "react";
 import { Link, useNavigate, useLocation } from "react-router-dom";
 import backend from "../../backend";
-import { GetJobStepRefResponse } from "../../backend/Api";
-import dashboard from "../../backend/Dashboard";
+import { EventData, EventSeverity, GetJobStepRefResponse, JobStepOutcome } from "../../backend/Api";
+import dashboard, { StatusColor } from "../../backend/Dashboard";
 import { ISideRailLink } from "../../base/components/SideRail";
 import { displayTimeZone, getElapsedString } from "../../base/utilities/timeUtils";
 import { ChangeButton } from "../ChangeButton";
+import { renderLine } from '../logs/LogRender';
 import { StepRefStatusIcon } from "../StatusIcon";
 import { JobDataView, JobDetailsV2 } from "./JobDetailsViewCommon";
 import { getHordeTheme } from "../../styles/theme";
@@ -19,6 +20,9 @@ import { getStepRefMessageInfo } from "./JobDetailsUtils";
 
 const sideRail: ISideRailLink = { text: "History", url: "rail_step_history" };
 
+/** Maximum number of recent failing/warning steps to fetch log events for */
+const MAX_STEPS_FOR_EVENTS = 30;
+
 class StepHistoryDataView extends JobDataView {
 
    filterUpdated() {
@@ -58,15 +62,69 @@ class StepHistoryDataView extends JobDataView {
       }
       backend.getJobStepHistory(jobData.streamId, stepName, 1024, jobData.templateId!).then(response => {
          this.history = response;
+         this.loadLogEvents();
       }).finally(() => {
          this.initialize(this.history?.length ? [sideRail] : undefined);
          this.updateReady();
       })
    }
 
+   async loadLogEvents() {
+      if (!this.history.length) {
+         return;
+      }
+
+      const candidates = this.history
+         .filter(ref => (ref.outcome === JobStepOutcome.Failure || ref.outcome === JobStepOutcome.Warnings) && ref.logId)
+         .slice(0, MAX_STEPS_FOR_EVENTS);
+
+      if (!candidates.length) {
+         return;
+      }
+
+      const failureLogIds = candidates.filter(ref => ref.outcome === JobStepOutcome.Failure).map(ref => ref.logId!);
+      const warningLogIds = candidates.filter(ref => ref.outcome === JobStepOutcome.Warnings).map(ref => ref.logId!);
+
+      this.loadingErrors = true;
+      this.updateReady();
+
+      try {
+         const requests: Promise<void>[] = [];
+
+         if (failureLogIds.length) {
+            requests.push(backend.getBatchLogEvents({ logIds: failureLogIds, severity: EventSeverity.Error }).then(response => {
+               for (const [logId, events] of Object.entries(response.logEvents)) {
+                  if (events.length) {
+                     this.logEvents.set(logId, events);
+                  }
+               }
+            }));
+         }
+
+         if (warningLogIds.length) {
+            requests.push(backend.getBatchLogEvents({ logIds: warningLogIds, severity: EventSeverity.Warning }).then(response => {
+               for (const [logId, events] of Object.entries(response.logEvents)) {
+                  if (events.length) {
+                     this.logEvents.set(logId, events);
+                  }
+               }
+            }));
+         }
+
+         await Promise.all(requests);
+      } catch (err) {
+         console.error("Failed to fetch batch log events:", err);
+      }
+
+      this.loadingErrors = false;
+      this.updateReady();
+   }
+
    clear() {
       this.history = [];
       this.stepId = undefined;
+      this.logEvents = new Map();
+      this.loadingErrors = false;
       super.clear();
    }
 
@@ -81,6 +139,10 @@ class StepHistoryDataView extends JobDataView {
 
    history: GetJobStepRefResponse[] = [];
 
+   logEvents: Map<string, EventData[]> = new Map();
+
+   loadingErrors = false;
+
    stepId?: string;
 
    order = 5;
@@ -105,7 +167,8 @@ export const StepHistoryPanel: React.FC<{ jobDetails: JobDetailsV2; stepId: stri
 
    dataView.subscribe();
 
-   const { hordeClasses } = getHordeStyling();
+   const { hordeClasses, modeColors } = getHordeStyling();
+   const [expandedLogs, setExpandedLogs] = useState<Set<string>>(new Set());
 
    if (!jobDetails.jobData) {
       return null;
@@ -119,13 +182,19 @@ export const StepHistoryPanel: React.FC<{ jobDetails: JobDetailsV2; stepId: stri
       return null;
    }
 
-   type HistoryItem = {
-      ref: GetJobStepRefResponse;
-   };
+   type HistoryItem =
+      | { kind: 'step'; ref: GetJobStepRefResponse }
+      | { kind: 'event'; ref: GetJobStepRefResponse; event: EventData };
 
-   const items: HistoryItem[] = dataView.history.map(h => {
-      return { ref: h }
-   });
+   const items: HistoryItem[] = [];
+   for (const h of dataView.history) {
+      items.push({ kind: 'step', ref: h });
+      if (h.logId && expandedLogs.has(h.logId) && dataView.logEvents.has(h.logId)) {
+         for (const event of dataView.logEvents.get(h.logId)!) {
+            items.push({ kind: 'event', ref: h, event });
+         }
+      }
+   }
 
    if (!items.length) {
       return null;
@@ -140,12 +209,53 @@ export const StepHistoryPanel: React.FC<{ jobDetails: JobDetailsV2; stepId: stri
       { key: 'column6', name: 'Duration', minWidth: 90, maxWidth: 90, isResizable: false },
    ];
 
+   const errorColor = dashboard.getStatusColors().get(StatusColor.Failure)!;
+   const warningColor = dashboard.getStatusColors().get(StatusColor.Warnings)!;
+
+   const textColorForBg = (bg: string): string => {
+      let r: number, g: number, b: number;
+      const rgbMatch = bg.match(/rgb\(\s*(\d+)\s*,\s*(\d+)\s*,\s*(\d+)\s*\)/);
+      if (rgbMatch) {
+         r = parseInt(rgbMatch[1]);
+         g = parseInt(rgbMatch[2]);
+         b = parseInt(rgbMatch[3]);
+      } else {
+         const hex = bg.replace("#", "");
+         r = parseInt(hex.substring(0, 2), 16);
+         g = parseInt(hex.substring(2, 4), 16);
+         b = parseInt(hex.substring(4, 6), 16);
+      }
+      const luminance = (0.299 * r + 0.587 * g + 0.114 * b) / 255;
+      return luminance > 0.5 ? "#000000" : "#FFFFFF";
+   };
+
    const renderItem = (item: HistoryItem, index?: number, column?: IColumn) => {
 
       if (!column) {
          return <div />;
       }
 
+      if (item.kind === 'event') {
+         if (column.key !== 'column1') {
+            return null;
+         }
+
+         const event = item.event;
+         const gutterColor = event.severity === EventSeverity.Warning ? warningColor : errorColor;
+         const logUrl = `/log/${event.logId}?lineIndex=${event.lineIndex + 1}`;
+         const lines = event.lines.filter(line => line.message?.trim().length).map((line, lineIdx) => (
+            <Stack key={`history_error_${event.logId}_${event.lineIndex}_${lineIdx}`} styles={{ root: { paddingLeft: 8, paddingRight: 8, lineBreak: "anywhere", whiteSpace: "pre-wrap", lineHeight: 18, fontSize: 10, fontFamily: "Horde Cousine Regular, monospace, monospace" } }}>
+               <Link style={{ color: modeColors.text }} to={logUrl}>{renderLine(navigate, line, undefined, {})}</Link>
+            </Stack>
+         ));
+
+         return (
+            <Stack style={{ borderLeft: `6px solid ${gutterColor}`, paddingLeft: 14, paddingTop: 4, paddingBottom: 4 }}>
+               {lines.length > 0 ? lines : <Text style={{ fontSize: 10 }}>Missing Log Data</Text>}
+            </Stack>
+         );
+      }
+
       const ref = item.ref;
 
       const step = jobDetails.stepById(stepId);
@@ -179,6 +289,36 @@ export const StepHistoryPanel: React.FC<{ jobDetails: JobDetailsV2; stepId: stri
 
       if (column.name === "Meta") {
          let { message: stepMessage, color: stepMessageColor } = getStepRefMessageInfo(ref);
+         const logEvents = ref.logId ? dataView.logEvents.get(ref.logId) : undefined;
+         const eventCount = logEvents?.length ?? 0;
+
+         if (eventCount > 0) {
+            const isExpanded = expandedLogs.has(ref.logId!);
+            const btnColor = ref.outcome === JobStepOutcome.Failure ? errorColor : warningColor;
+            const btnTextColor = textColorForBg(btnColor);
+            return <Stack verticalAlign="center" verticalFill horizontal horizontalAlign="center">
+                  <DefaultButton
+                     styles={{
+                        root: { minWidth: 0, height: 20, padding: "0 6px", backgroundColor: btnColor, borderColor: btnColor, borderRadius: 2, color: btnTextColor },
+                        rootHovered: { backgroundColor: btnColor, borderColor: btnColor, color: btnTextColor, opacity: 0.9 },
+                        rootPressed: { backgroundColor: btnColor, borderColor: btnColor, color: btnTextColor, opacity: 0.8 },
+                        label: { fontSize: 10, fontFamily: "Horde Open Sans SemiBold", color: btnTextColor, lineHeight: "18px", margin: 0, padding: 0 }
+                     }}
+                     onClick={(ev) => {
+                        ev.preventDefault();
+                        ev.stopPropagation();
+                        const next = new Set(expandedLogs);
+                        if (isExpanded) {
+                           next.delete(ref.logId!);
+                        } else {
+                           next.add(ref.logId!);
+                        }
+                        setExpandedLogs(next);
+                     }}
+                  ><FontIcon iconName={isExpanded ? "ChevronDown" : "ChevronRight"} style={{ fontSize: 8, marginRight: 4, position: "relative", top: 1, color: btnTextColor }} />{stepMessage} ({eventCount})</DefaultButton>
+               </Stack>
+         }
+
          return <Stack verticalAlign="center" horizontal tokens={{ childrenGap: 6 }} horizontalAlign="center">
                {!!stepMessageColor && <FontIcon style={{ fontSize: 13, color: stepMessageColor, paddingTop: 1 }} iconName="Square" />}
                <Text style={{fontSize: 12, overflow: "hidden", textOverflow: "ellipsis"}}>{stepMessage}</Text>
@@ -227,6 +367,13 @@ export const StepHistoryPanel: React.FC<{ jobDetails: JobDetailsV2; stepId: stri
       if (props) {
 
          const item = props!.item as HistoryItem;
+
+         if (item.kind === 'event') {
+            const commonSelectors = { ".ms-DetailsRow-cell": { "overflow": "visible" } };
+            props.styles = { ...props.styles, root: { background: theme.palette.neutralLighterAlt, selectors: { ...commonSelectors as any, "&:hover": { background: theme.palette.neutralLighter } } } };
+            return <div className="job-item"><DetailsRow {...props} /></div>;
+         }
+
          const ref = item.ref;
          const url = `/job/${ref.jobId}?step=${ref.stepId}`;
 
@@ -247,7 +394,10 @@ export const StepHistoryPanel: React.FC<{ jobDetails: JobDetailsV2; stepId: stri
    return (<Stack id={sideRail.url} styles={{ root: { paddingTop: 18, paddingRight: 12 } }}>
       <Stack className={hordeClasses.raised}>
          <Stack tokens={{ childrenGap: 12 }}>
-            <Text variant="mediumPlus" styles={{ root: { fontFamily: "Horde Open Sans SemiBold" } }}>{"History"}</Text>
+            <Stack horizontal tokens={{ childrenGap: 18 }}>
+               <Text variant="mediumPlus" styles={{ root: { fontFamily: "Horde Open Sans SemiBold" } }}>{"History"}</Text>
+               {dataView.loadingErrors && <Spinner size={SpinnerSize.small} />}
+            </Stack>
             <Stack styles={{ root: { paddingLeft: 4, paddingRight: 0, paddingTop: 8, paddingBottom: 4 } }}>
                <div style={{ overflowY: 'auto', overflowX: 'hidden', maxHeight: "400px" }} data-is-scrollable={true}>
                   {!!dataView.history.length && <DetailsList
diff --git a/Engine/Source/Programs/Horde/Plugins/Compute/HordeServer.Compute/Logs/LogsController.cs b/Engine/Source/Programs/Horde/Plugins/Compute/HordeServer.Compute/Logs/LogsController.cs
--- a/Engine/Source/Programs/Horde/Plugins/Compute/HordeServer.Compute/Logs/LogsController.cs
+++ b/Engine/Source/Programs/Horde/Plugins/Compute/HordeServer.Compute/Logs/LogsController.cs
@@ -372,6 +372,73 @@ namespace HordeServer.Logs
 			return responses;
 		}
 
+		/// <summary>
+		/// Retrieve events for multiple logs in a single request
+		/// </summary>
+		/// <param name="request">Request containing log ids to fetch events for</param>
+		/// <param name="cancellationToken">Cancellation token for the request</param>
+		/// <returns>Events grouped by log id</returns>
+		[HttpPost]
+		[Route("/api/v1/logs/events")]
+		[ProducesResponseType(typeof(GetBatchLogEventsResponse), 200)]
+		public async Task<ActionResult<GetBatchLogEventsResponse>> GetBatchEventsAsync([FromBody] GetBatchLogEventsRequest request, CancellationToken cancellationToken = default)
+		{
+			GetBatchLogEventsResponse response = new GetBatchLogEventsResponse();
+
+			foreach (string logIdStr in request.LogIds)
+			{
+				LogId logId = LogId.Parse(logIdStr);
+				ILog? log = await _logCollection.GetAsync(logId, cancellationToken);
+				if (log == null)
+				{
+					continue;
+				}
+				if (!await AuthorizeAsync(log, LogAclAction.ViewLog, User, cancellationToken))
+				{
+					continue;
+				}
+
+				List<ILogAnchor> anchors = await log.GetAnchorsAsync(null, null, null, cancellationToken);
+
+				Dictionary<ObjectId, int?> spanIdToIssueId = new Dictionary<ObjectId, int?>();
+				List<GetLogEventResponse> events = new List<GetLogEventResponse>();
+
+				foreach (ILogAnchor anchor in anchors)
+				{
+					if (request.Severity != null && anchor.Severity != request.Severity.Value)
+					{
+						continue;
+					}
+
+					ILogEventData logEventData = await anchor.GetDataAsync(cancellationToken);
+
+					int? issueId = null;
+					if (anchor.SpanId != null && !spanIdToIssueId.TryGetValue(anchor.SpanId.Value, out issueId))
+					{
+						foreach (ILogExtIssueProvider issueProvider in _issueProviders)
+						{
+							issueId ??= await issueProvider.GetIssueIdAsync(anchor.SpanId.Value, cancellationToken);
+						}
+						spanIdToIssueId[anchor.SpanId.Value] = issueId;
+					}
+
+					events.Add(CreateGetLogEventResponse(anchor, logEventData, issueId));
+
+					if (request.Count != null && events.Count >= request.Count.Value)
+					{
+						break;
+					}
+				}
+
+				if (events.Count > 0)
+				{
+					response.LogEvents[logIdStr] = events;
+				}
+			}
+
+			return response;
+		}
+
 		/// <summary>
 		/// Create a log event response message
 		/// </summary>
diff --git a/Engine/Source/Programs/Shared/EpicGames.Horde/Logs/LogEventMessages.cs b/Engine/Source/Programs/Shared/EpicGames.Horde/Logs/LogEventMessages.cs
--- a/Engine/Source/Programs/Shared/EpicGames.Horde/Logs/LogEventMessages.cs
+++ b/Engine/Source/Programs/Shared/EpicGames.Horde/Logs/LogEventMessages.cs
@@ -33,6 +33,38 @@ namespace EpicGames.Horde.Logs
 		Error = 3,
 	}
 
+	/// <summary>
+	/// Request to get events for multiple logs at once
+	/// </summary>
+	public class GetBatchLogEventsRequest
+	{
+		/// <summary>
+		/// Log ids to fetch events for
+		/// </summary>
+		public List<string> LogIds { get; set; } = [];
+
+		/// <summary>
+		/// Maximum number of events to return per log. If null, returns all events.
+		/// </summary>
+		public int? Count { get; set; }
+
+		/// <summary>
+		/// Severity filter. If specified, only events matching this severity are returned.
+		/// </summary>
+		public LogEventSeverity? Severity { get; set; }
+	}
+
+	/// <summary>
+	/// Response containing events grouped by log id
+	/// </summary>
+	public class GetBatchLogEventsResponse
+	{
+		/// <summary>
+		/// Events grouped by log id
+		/// </summary>
+		public Dictionary<string, List<GetLogEventResponse>> LogEvents { get; set; } = [];
+	}
+
 	/// <summary>
 	/// Information about an uploaded event
 	/// </summary>
