Horde: Fix up and unify change list preflight display

diff --git a/Engine/Source/Programs/Horde/HordeDashboard/src/backend/JobHandler.ts b/Engine/Source/Programs/Horde/HordeDashboard/src/backend/JobHandler.ts
--- a/Engine/Source/Programs/Horde/HordeDashboard/src/backend/JobHandler.ts
+++ b/Engine/Source/Programs/Horde/HordeDashboard/src/backend/JobHandler.ts
@@ -79,7 +79,7 @@ export class JobHandler {
         try {
 
             // discover new/updated jobs
-            let filter = "id,streamId,name,change,preflightChange,templateId,templateHash,graphHash,startedByUserInfo,abortedByUserInfo,createTime,state,arguments,updateTime,labels,defaultLabel,preflightDescription,cancellationReason";
+            let filter = "id,streamId,name,change,preflightChange,preflightCommitId,templateId,templateHash,graphHash,startedByUserInfo,abortedByUserInfo,createTime,state,arguments,updateTime,labels,defaultLabel,preflightDescription,cancellationReason";
 
             if (this.includeBatches) {
                 filter += ",batches";
diff --git a/Engine/Source/Programs/Horde/HordeDashboard/src/base/utilities/commitUtils.ts b/Engine/Source/Programs/Horde/HordeDashboard/src/base/utilities/commitUtils.ts
new file mode 100644
--- /dev/null
+++ b/Engine/Source/Programs/Horde/HordeDashboard/src/base/utilities/commitUtils.ts
@@ -0,0 +1,33 @@
+// Copyright Epic Games, Inc. All Rights Reserved.
+
+/**
+ * Resolves a preflight display string from the available fields.
+ * Prefers preflightCommitId; returns the numeric value if it parses as an integer,
+ * otherwise an abbreviated 8-char hash. Falls back to preflightChange for legacy data.
+ * Returns undefined if no preflight info is available.
+ */
+export function getPreflightDisplay(preflightCommitId?: string, preflightChange?: number): string | undefined {
+	if (preflightCommitId) {
+		const num = parseInt(preflightCommitId);
+		return !isNaN(num) ? num.toString() : preflightCommitId.substring(0, 8);
+	}
+	if (preflightChange && preflightChange > 0) {
+		return preflightChange.toString();
+	}
+	return undefined;
+}
+
+/**
+ * Returns the numeric CL for a preflight, or undefined if the preflight uses a non-numeric commit ID.
+ * Prefers preflightCommitId; falls back to preflightChange for legacy data.
+ */
+export function getPreflightCL(preflightCommitId?: string, preflightChange?: number): number | undefined {
+	if (preflightCommitId) {
+		const num = parseInt(preflightCommitId);
+		return !isNaN(num) ? num : undefined;
+	}
+	if (preflightChange && preflightChange > 0) {
+		return preflightChange;
+	}
+	return undefined;
+}
diff --git a/Engine/Source/Programs/Horde/HordeDashboard/src/components/ChangeButton.tsx b/Engine/Source/Programs/Horde/HordeDashboard/src/components/ChangeButton.tsx
--- a/Engine/Source/Programs/Horde/HordeDashboard/src/components/ChangeButton.tsx
+++ b/Engine/Source/Programs/Horde/HordeDashboard/src/components/ChangeButton.tsx
@@ -5,6 +5,7 @@ import React, { MutableRefObject, useState } from 'react';
 import { GetBatchResponse, GetChangeSummaryResponse, GetJobStepRefResponse, GetThinUserInfoResponse, JobStepBatchError, JobStepBatchState, JobStepOutcome, JobStepState } from '../backend/Api';
 import dashboard, { StatusColor } from "../backend/Dashboard";
 import { projectStore } from '../backend/ProjectStore';
+import { getPreflightDisplay } from '../base/utilities/commitUtils';
 
 export type ChangeContextMenuTarget = {
 
@@ -19,6 +20,7 @@ export type JobParameters = {
    streamId: string;
    change?: number;
    preflightChange?: number;
+   preflightCommitId?: string;
    preflightDescription?: string;
    abortedByUserInfo?: GetThinUserInfoResponse;
    cancellationReason?: string;
@@ -40,8 +42,9 @@ export const ChangeContextMenu: React.FC<{ target: ChangeContextMenuTarget, onDi
    }
 
    let change = jobChange.toString();
-   if (!stepRef && job?.preflightChange) {
-      change = `PF ${job.preflightChange}`;
+   const pfMenuDisplay = !stepRef ? getPreflightDisplay(job?.preflightCommitId, job?.preflightChange) : undefined;
+   if (pfMenuDisplay) {
+      change = `PF ${pfMenuDisplay}`;
    }
 
    const copyToClipboard = (value: string | undefined) => {
@@ -123,7 +126,7 @@ export const ChangeContextMenu: React.FC<{ target: ChangeContextMenuTarget, onDi
    menuItems.push({
       key: 'copy_to_clipboard',
       text: 'Copy CL to Clipboard',
-      onClick: () => copyToClipboard(change.replace("PF ", ""))
+      onClick: () => copyToClipboard(job?.preflightCommitId ?? change.replace("PF ", ""))
    });
 
    if (dashboard.swarmUrl) {
@@ -265,8 +268,9 @@ export const ChangeButton: React.FC<{ job?: JobParameters, stepRef?: GetJobStepR
    const { text, color } = getJobSummary(job);
 
    let change = job?.change?.toString() ?? "Latest";
-   if (job?.preflightChange) {
-      change = `PF ${job.preflightChange}`;
+   const pfDisplay = getPreflightDisplay(job?.preflightCommitId, job?.preflightChange);
+   if (pfDisplay) {
+      change = `PF ${pfDisplay}`;
    }
 
    if (stepRef) {
diff --git a/Engine/Source/Programs/Horde/HordeDashboard/src/components/JobOperationsContextMenu.tsx b/Engine/Source/Programs/Horde/HordeDashboard/src/components/JobOperationsContextMenu.tsx
--- a/Engine/Source/Programs/Horde/HordeDashboard/src/components/JobOperationsContextMenu.tsx
+++ b/Engine/Source/Programs/Horde/HordeDashboard/src/components/JobOperationsContextMenu.tsx
@@ -5,6 +5,7 @@ import React, { useState } from 'react';
 import { GetJobResponse, JobState } from '../backend/Api';
 import dashboard from '../backend/Dashboard';
 import { projectStore } from '../backend/ProjectStore';
+import { getPreflightDisplay } from '../base/utilities/commitUtils';
 import { AbortJobModal } from './jobDetailsV2/AbortJobModal';
 // import { EditJobModal } from './EditJobModal';
 
@@ -33,8 +34,9 @@ export const JobOperationsContextMenu: React.FC<{ job: GetJobResponse, children?
    const pinned = dashboard.jobPinned(job.id);
 
    let change = job.change!.toString();
-   if (job.preflightChange) {
-      change = job.preflightChange.toString();
+   const pfDisplay = getPreflightDisplay(job.preflightCommitId, job.preflightChange);
+   if (pfDisplay) {
+      change = pfDisplay;
    }
 
    let menuItems: IContextualMenuItem[] = [
diff --git a/Engine/Source/Programs/Horde/HordeDashboard/src/components/JobSearchSimple.tsx b/Engine/Source/Programs/Horde/HordeDashboard/src/components/JobSearchSimple.tsx
--- a/Engine/Source/Programs/Horde/HordeDashboard/src/components/JobSearchSimple.tsx
+++ b/Engine/Source/Programs/Horde/HordeDashboard/src/components/JobSearchSimple.tsx
@@ -197,7 +197,7 @@ export const JobSearchSimpleModal: React.FC<{ streamIdIn: string, changeListIn?:
                id: jobIds.length ? jobIds : undefined,
                streamId: searchStreamId?.toLowerCase(),
                name: name,
-               filter: "id,labels,defaultLabel,name,change,createTime,streamId,preflightChange,startedByUserInfo,arguments,abortedByUserInfo,state,batches,cancellationReason",
+               filter: "id,labels,defaultLabel,name,change,createTime,streamId,preflightChange,preflightCommitId,startedByUserInfo,arguments,abortedByUserInfo,state,batches,cancellationReason",
                minChange: minCL,
                maxChange: maxCL,
                t: (!state.global && c_templates.length && !state.containsStep) ? c_templates.map(t => t.id) : undefined,
diff --git a/Engine/Source/Programs/Horde/HordeDashboard/src/components/JobViewIncremental.tsx b/Engine/Source/Programs/Horde/HordeDashboard/src/components/JobViewIncremental.tsx
--- a/Engine/Source/Programs/Horde/HordeDashboard/src/components/JobViewIncremental.tsx
+++ b/Engine/Source/Programs/Horde/HordeDashboard/src/components/JobViewIncremental.tsx
@@ -10,6 +10,7 @@ import { GetChangeSummaryResponse, GetJobsTabLabelColumnResponse, GetJobsTabPara
 import { CommitCache } from '../backend/CommitCache';
 import dashboard from '../backend/Dashboard';
 import { JobHandler } from '../backend/JobHandler';
+import { getPreflightDisplay } from '../base/utilities/commitUtils';
 import { filterJob, JobFilterSimple } from '../base/utilities/filter';
 import { displayTimeZone } from '../base/utilities/timeUtils';
 import { getJobStateColor, getLabelColor } from '../styles/colors';
@@ -514,10 +515,9 @@ const JobList: React.FC<{ tab: string; filter: JobFilterSimple; aliasTemplates?:
          }
 
          let change = item.job?.change?.toString() ?? "Latest";
-         if (item.job?.preflightChange) {
-            change = `PF ${item.job.preflightChange}`;
-         } else if (item.job?.preflightCommitId) {
-            change = `PF ${item.job.preflightCommitId.substring(0, 8)}`;
+         const pfDisplay = getPreflightDisplay(item.job?.preflightCommitId, item.job?.preflightChange);
+         if (pfDisplay) {
+            change = `PF ${pfDisplay}`;
          }
 
 
diff --git a/Engine/Source/Programs/Horde/HordeDashboard/src/components/PoolView.tsx b/Engine/Source/Programs/Horde/HordeDashboard/src/components/PoolView.tsx
--- a/Engine/Source/Programs/Horde/HordeDashboard/src/components/PoolView.tsx
+++ b/Engine/Source/Programs/Horde/HordeDashboard/src/components/PoolView.tsx
@@ -151,7 +151,7 @@ class PoolHandler extends PollBase {
 
          if (streamJobIds.size) {
 
-            const streamJobs = await backend.getJobs({ id: Array.from(streamJobIds), filter: "id,name,batches,change,createTime,streamId,preflightChange" });
+            const streamJobs = await backend.getJobs({ id: Array.from(streamJobIds), filter: "id,name,batches,change,createTime,streamId,preflightChange,preflightCommitId" });
 
             streamJobs.forEach(job => {
 
@@ -209,7 +209,7 @@ class PoolHandler extends PollBase {
 
             const query: JobQuery = {
                id: jobIds,
-               filter: "id,name,batches,change,createTime,streamId,preflightChange",
+               filter: "id,name,batches,change,createTime,streamId,preflightChange,preflightCommitId",
             }
 
             const jobs = await backend.getJobs(query);
diff --git a/Engine/Source/Programs/Horde/HordeDashboard/src/components/StreamView.tsx b/Engine/Source/Programs/Horde/HordeDashboard/src/components/StreamView.tsx
--- a/Engine/Source/Programs/Horde/HordeDashboard/src/components/StreamView.tsx
+++ b/Engine/Source/Programs/Horde/HordeDashboard/src/components/StreamView.tsx
@@ -118,7 +118,7 @@ class IncrementalState {
          const batch = rincrementals.slice(0, 5);
 
          await Promise.all(batch.map(b => {
-            return backend.getStreamJobs(b.streamId, { template: [b.template.id], count: 5, filter: "id,labels,createTime,streamId,defaultLabel,preflightChange" })
+            return backend.getStreamJobs(b.streamId, { template: [b.template.id], count: 5, filter: "id,labels,createTime,streamId,defaultLabel,preflightChange,preflightCommitId" })
             // eslint-disable-next-line no-loop-func
          })).then((r) => {
             for (let i = 0; i < r.length; i++) {
diff --git a/Engine/Source/Programs/Horde/HordeDashboard/src/components/UserHome.tsx b/Engine/Source/Programs/Horde/HordeDashboard/src/components/UserHome.tsx
--- a/Engine/Source/Programs/Horde/HordeDashboard/src/components/UserHome.tsx
+++ b/Engine/Source/Programs/Horde/HordeDashboard/src/components/UserHome.tsx
@@ -412,7 +412,7 @@ class UserJobsHandler {
       const maxJobs = 100;
 
       try {
-         let filter = "id,streamId,name,change,preflightChange,templateId,templateHash,graphHash,startedByUserInfo,createTime,state,arguments,updateTime,labels,defaultLabel,batches,autoSubmit,autoSubmitChange,preflightDescription,abortedByUserInfo,cancellationReason";
+         let filter = "id,streamId,name,change,preflightChange,preflightCommitId,templateId,templateHash,graphHash,startedByUserInfo,createTime,state,arguments,updateTime,labels,defaultLabel,batches,autoSubmit,autoSubmitChange,preflightDescription,abortedByUserInfo,cancellationReason";
 
          const query: JobQuery = {
             filter: filter,
diff --git a/Engine/Source/Programs/Horde/HordeDashboard/src/components/jobDetailsV2/JobDetailPreflight.tsx b/Engine/Source/Programs/Horde/HordeDashboard/src/components/jobDetailsV2/JobDetailPreflight.tsx
--- a/Engine/Source/Programs/Horde/HordeDashboard/src/components/jobDetailsV2/JobDetailPreflight.tsx
+++ b/Engine/Source/Programs/Horde/HordeDashboard/src/components/jobDetailsV2/JobDetailPreflight.tsx
@@ -7,6 +7,7 @@ import backend from '../../backend';
 import { JobState } from '../../backend/Api';
 import dashboard from '../../backend/Dashboard';
 import { ISideRailLink } from '../../base/components/SideRail';
+import { getPreflightCL, getPreflightDisplay } from '../../base/utilities/commitUtils';
 import { errorDialogStore } from '../error/ErrorStore';
 import { JobDataView, JobDetailsV2 } from "./JobDetailsViewCommon";
 import { getHordeStyling } from '../../styles/Styles';
@@ -70,10 +71,11 @@ const AutosubmitInfo: React.FC<{ jobDetails: JobDetailsV2 }> = observer(({ jobDe
 
    if (jobData.state !== JobState.Complete) {
 
-      if (jobData.preflightChange) {
+      const preflightCL = getPreflightCL(jobData.preflightCommitId, jobData.preflightChange);
+      if (preflightCL) {
          return <Stack tokens={{ childrenGap: 8 }} style={{ paddingBottom: 12, paddingTop: 18 }}>
             <Stack style={{ paddingTop: 8 }}>
-               <Checkbox label={`Automatically submit changelist ${jobData.preflightChange} upon preflight success`}
+               <Checkbox label={`Automatically submit changelist ${preflightCL} upon preflight success`}
                   checked={state.autoSubmit}
                   disabled={state.inflight || (jobData.startedByUserInfo?.id !== dashboard.userId)}
                   onChange={(ev, checked) => {
@@ -132,7 +134,10 @@ export const PreflightPanel: React.FC<{ jobDetails: JobDetailsV2 }> = observer((
       return null;
    }
 
-   const hasPreflight = !!jobData.preflightChange || !!jobData.preflightDescription || !!jobData.preflightCommitId;
+   const pfDisplayCL = getPreflightCL(jobData.preflightCommitId, jobData.preflightChange);
+   const pfDisplayText = getPreflightDisplay(jobData.preflightCommitId, jobData.preflightChange);
+   const pfIsHash = jobData.preflightCommitId && !pfDisplayCL;
+   const hasPreflight = !!pfDisplayText || !!jobData.preflightDescription;
 
    if (!preflightView.initialized) {
       preflightView.initialize(hasPreflight ? [sideRail] : undefined);
@@ -147,16 +152,16 @@ export const PreflightPanel: React.FC<{ jobDetails: JobDetailsV2 }> = observer((
          <Stack tokens={{ childrenGap: 12 }} grow>
             <Stack horizontal>
                <Stack>
-                  {jobData.preflightChange && (
-                     <Text variant="mediumPlus" styles={{ root: { fontFamily: "Horde Open Sans SemiBold" } }}>{`Preflight (CL ${jobData.preflightChange})`}</Text>
+                  {pfDisplayCL && (
+                     <Text variant="mediumPlus" styles={{ root: { fontFamily: "Horde Open Sans SemiBold" } }}>{`Preflight (CL ${pfDisplayCL})`}</Text>
                   )}
-                  {jobData.preflightCommitId && !jobData.preflightChange && (
+                  {pfIsHash && (
                      <Stack horizontal tokens={{ childrenGap: 8 }} verticalAlign="center">
                         <Text variant="mediumPlus" styles={{ root: { fontFamily: "Horde Open Sans SemiBold" } }}>
                            Preflight Commit:
                         </Text>
                         <Text variant="medium" styles={{ root: { fontFamily: 'monospace' } }}>
-                           {jobData.preflightCommitId.substring(0, 12)}...
+                           {jobData.preflightCommitId!.substring(0, 12)}...
                         </Text>
                         <IconButton
                            iconProps={{ iconName: commitIdCopied ? 'CheckMark' : 'Copy' }}
diff --git a/Engine/Source/Programs/Horde/HordeDashboard/src/components/jobDetailsV2/JobDetailViewV2.tsx b/Engine/Source/Programs/Horde/HordeDashboard/src/components/jobDetailsV2/JobDetailViewV2.tsx
--- a/Engine/Source/Programs/Horde/HordeDashboard/src/components/jobDetailsV2/JobDetailViewV2.tsx
+++ b/Engine/Source/Programs/Horde/HordeDashboard/src/components/jobDetailsV2/JobDetailViewV2.tsx
@@ -9,6 +9,7 @@ import { useWindowSize } from "../../base/utilities/hooks";
 import { getHordeStyling } from "../../styles/Styles";
 import { BreadcrumbItem, Breadcrumbs } from "../Breadcrumbs";
 import { useQuery } from "horde/base/utilities/hooks";
+import { getPreflightDisplay } from "../../base/utilities/commitUtils";
 import { getStreamDisplayName, getStreamLink } from "../../base/utilities/streamUtils";
 import { TopNav } from "../TopNav";
 import { JobArtifactsPanel } from "./JobDetailArtifacts";
@@ -95,8 +96,9 @@ const JobBreadCrumbs: React.FC<{ jobDetails: JobDetailsV2 }> = observer(({ jobDe
 
 
    let clText = "";
-   if (jobData.preflightChange) {
-      clText = `Preflight ${jobData.preflightChange} `;
+   const preflightDisplay = getPreflightDisplay(jobData.preflightCommitId, jobData.preflightChange);
+   if (preflightDisplay) {
+      clText = `Preflight ${preflightDisplay} `;
       clText += ` - Base ${jobData.change ? jobData.change : "Latest"}`;
 
    } else {
