Horde: Fix preflight commit ID parsing

diff --git a/Engine/Source/Programs/Horde/HordeDashboard/src/base/utilities/commitUtils.ts b/Engine/Source/Programs/Horde/HordeDashboard/src/base/utilities/commitUtils.ts
--- a/Engine/Source/Programs/Horde/HordeDashboard/src/base/utilities/commitUtils.ts
+++ b/Engine/Source/Programs/Horde/HordeDashboard/src/base/utilities/commitUtils.ts
@@ -2,14 +2,13 @@
 
 /**
  * Resolves a preflight display string from the available fields.
- * Prefers preflightCommitId; returns the numeric value if it parses as an integer,
+ * Prefers preflightCommitId; returns the numeric value if the entire string is digits,
  * otherwise an abbreviated 8-char hash. Falls back to preflightChange for legacy data.
  * Returns undefined if no preflight info is available.
  */
 export function getPreflightDisplay(preflightCommitId?: string, preflightChange?: number): string | undefined {
 	if (preflightCommitId) {
-		const num = parseInt(preflightCommitId);
-		return !isNaN(num) ? num.toString() : preflightCommitId.substring(0, 8);
+		return /^\d+$/.test(preflightCommitId) ? preflightCommitId : preflightCommitId.substring(0, 8);
 	}
 	if (preflightChange && preflightChange > 0) {
 		return preflightChange.toString();
@@ -23,8 +22,7 @@ export function getPreflightDisplay(preflightCommitId?: string, preflightChange?
  */
 export function getPreflightCL(preflightCommitId?: string, preflightChange?: number): number | undefined {
 	if (preflightCommitId) {
-		const num = parseInt(preflightCommitId);
-		return !isNaN(num) ? num : undefined;
+		return /^\d+$/.test(preflightCommitId) ? Number(preflightCommitId) : undefined;
 	}
 	if (preflightChange && preflightChange > 0) {
 		return preflightChange;
diff --git a/Engine/Source/Programs/Horde/Plugins/Build/HordeServer.Build/Jobs/JobsController.cs b/Engine/Source/Programs/Horde/Plugins/Build/HordeServer.Build/Jobs/JobsController.cs
--- a/Engine/Source/Programs/Horde/Plugins/Build/HordeServer.Build/Jobs/JobsController.cs
+++ b/Engine/Source/Programs/Horde/Plugins/Build/HordeServer.Build/Jobs/JobsController.cs
@@ -180,6 +180,9 @@ namespace HordeServer.Jobs
 					{
 						ICommit preflightCommit = await commitCollection.GetAsync(create.PreflightCommitId, cancellationToken);
 						shelfInfo = new ShelfInfo(preflightCommit.Description, new List<CommitTag>());
+
+						// Normalize to the canonical commit ID (revision number) for consistent storage and display
+						create.PreflightCommitId = new CommitId(preflightCommit.Id.Name);
 					}
 					catch (CommitCollectionException)
 					{
