Horde: Add batch log events API and inline error/warning display in step history
New POST /api/v1/logs/events endpoint for fetching events across multiple logs
in a single request with severity filtering and count limiting. Step history
panel shows expandable error/warning events inline with log links.

diff --git a/Engine/Source/Programs/Horde/HordeDashboard/src/backend/Api.ts b/Engine/Source/Programs/Horde/HordeDashboard/src/backend/Api.ts
--- a/Engine/Source/Programs/Horde/HordeDashboard/src/backend/Api.ts
+++ b/Engine/Source/Programs/Horde/HordeDashboard/src/backend/Api.ts
@@ -381,8 +381,13 @@ export type GetBatchLogEventsRequest = {
 	severity?: EventSeverity;
 }
 
+export type BatchLogEventsEntry = {
+	total: number;
+	events: EventData[];
+}
+
 export type GetBatchLogEventsResponse = {
-	logEvents: Record<string, EventData[]>;
+	logEvents: Record<string, BatchLogEventsEntry>;
 }
 
 export type UsersQuery = {
diff --git a/Engine/Source/Programs/Horde/HordeDashboard/src/components/jobDetailsV2/JobDetailStepHistory.tsx b/Engine/Source/Programs/Horde/HordeDashboard/src/components/jobDetailsV2/JobDetailStepHistory.tsx
--- a/Engine/Source/Programs/Horde/HordeDashboard/src/components/jobDetailsV2/JobDetailStepHistory.tsx
+++ b/Engine/Source/Programs/Horde/HordeDashboard/src/components/jobDetailsV2/JobDetailStepHistory.tsx
@@ -6,7 +6,7 @@ import moment from "moment";
 import { useEffect, useState } from "react";
 import { Link, useNavigate, useLocation } from "react-router-dom";
 import backend from "../../backend";
-import { EventData, EventSeverity, GetJobStepRefResponse, JobStepOutcome } from "../../backend/Api";
+import { BatchLogEventsEntry, EventData, EventSeverity, GetJobStepRefResponse, JobStepOutcome } from "../../backend/Api";
 import dashboard, { StatusColor } from "../../backend/Dashboard";
 import { ISideRailLink } from "../../base/components/SideRail";
 import { displayTimeZone, getElapsedString } from "../../base/utilities/timeUtils";
@@ -92,20 +92,20 @@ class StepHistoryDataView extends JobDataView {
          const requests: Promise<void>[] = [];
 
          if (failureLogIds.length) {
-            requests.push(backend.getBatchLogEvents({ logIds: failureLogIds, severity: EventSeverity.Error }).then(response => {
-               for (const [logId, events] of Object.entries(response.logEvents)) {
-                  if (events.length) {
-                     this.logEvents.set(logId, events);
+            requests.push(backend.getBatchLogEvents({ logIds: failureLogIds, severity: EventSeverity.Error, count: 20 }).then(response => {
+               for (const [logId, entry] of Object.entries(response.logEvents)) {
+                  if (entry.events.length) {
+                     this.logEvents.set(logId, entry);
                   }
                }
             }));
          }
 
          if (warningLogIds.length) {
-            requests.push(backend.getBatchLogEvents({ logIds: warningLogIds, severity: EventSeverity.Warning }).then(response => {
-               for (const [logId, events] of Object.entries(response.logEvents)) {
-                  if (events.length) {
-                     this.logEvents.set(logId, events);
+            requests.push(backend.getBatchLogEvents({ logIds: warningLogIds, severity: EventSeverity.Warning, count: 20 }).then(response => {
+               for (const [logId, entry] of Object.entries(response.logEvents)) {
+                  if (entry.events.length) {
+                     this.logEvents.set(logId, entry);
                   }
                }
             }));
@@ -139,7 +139,7 @@ class StepHistoryDataView extends JobDataView {
 
    history: GetJobStepRefResponse[] = [];
 
-   logEvents: Map<string, EventData[]> = new Map();
+   logEvents: Map<string, BatchLogEventsEntry> = new Map();
 
    loadingErrors = false;
 
@@ -184,15 +184,20 @@ export const StepHistoryPanel: React.FC<{ jobDetails: JobDetailsV2; stepId: stri
 
    type HistoryItem =
       | { kind: 'step'; ref: GetJobStepRefResponse }
-      | { kind: 'event'; ref: GetJobStepRefResponse; event: EventData };
+      | { kind: 'event'; ref: GetJobStepRefResponse; event: EventData }
+      | { kind: 'overflow'; ref: GetJobStepRefResponse; remaining: number };
 
    const items: HistoryItem[] = [];
    for (const h of dataView.history) {
       items.push({ kind: 'step', ref: h });
       if (h.logId && expandedLogs.has(h.logId) && dataView.logEvents.has(h.logId)) {
-         for (const event of dataView.logEvents.get(h.logId)!) {
+         const entry = dataView.logEvents.get(h.logId)!;
+         for (const event of entry.events) {
             items.push({ kind: 'event', ref: h, event });
          }
+         if (entry.total > entry.events.length) {
+            items.push({ kind: 'overflow', ref: h, remaining: entry.total - entry.events.length });
+         }
       }
    }
 
@@ -256,6 +261,18 @@ export const StepHistoryPanel: React.FC<{ jobDetails: JobDetailsV2; stepId: stri
          );
       }
 
+      if (item.kind === 'overflow') {
+         if (column.key !== 'column1') {
+            return null;
+         }
+         const logUrl = `/log/${item.ref.logId}`;
+         return (
+            <Stack style={{ paddingLeft: 20, paddingTop: 4, paddingBottom: 4 }}>
+               <Link to={logUrl} style={{ fontSize: 10, color: modeColors.text }}>View {item.remaining} more in full log...</Link>
+            </Stack>
+         );
+      }
+
       const ref = item.ref;
 
       const step = jobDetails.stepById(stepId);
@@ -289,8 +306,8 @@ export const StepHistoryPanel: React.FC<{ jobDetails: JobDetailsV2; stepId: stri
 
       if (column.name === "Meta") {
          let { message: stepMessage, color: stepMessageColor } = getStepRefMessageInfo(ref);
-         const logEvents = ref.logId ? dataView.logEvents.get(ref.logId) : undefined;
-         const eventCount = logEvents?.length ?? 0;
+         const logEntry = ref.logId ? dataView.logEvents.get(ref.logId) : undefined;
+         const eventCount = logEntry?.total ?? 0;
 
          if (eventCount > 0) {
             const isExpanded = expandedLogs.has(ref.logId!);
@@ -368,7 +385,7 @@ export const StepHistoryPanel: React.FC<{ jobDetails: JobDetailsV2; stepId: stri
 
          const item = props!.item as HistoryItem;
 
-         if (item.kind === 'event') {
+         if (item.kind === 'event' || item.kind === 'overflow') {
             const commonSelectors = { ".ms-DetailsRow-cell": { "overflow": "visible" } };
             props.styles = { ...props.styles, root: { background: theme.palette.neutralLighterAlt, selectors: { ...commonSelectors as any, "&:hover": { background: theme.palette.neutralLighter } } } };
             return <div className="job-item"><DetailsRow {...props} /></div>;
diff --git a/Engine/Source/Programs/Horde/Plugins/Compute/HordeServer.Compute/Logs/LogsController.cs b/Engine/Source/Programs/Horde/Plugins/Compute/HordeServer.Compute/Logs/LogsController.cs
--- a/Engine/Source/Programs/Horde/Plugins/Compute/HordeServer.Compute/Logs/LogsController.cs
+++ b/Engine/Source/Programs/Horde/Plugins/Compute/HordeServer.Compute/Logs/LogsController.cs
@@ -400,16 +400,27 @@ namespace HordeServer.Logs
 
 				List<ILogAnchor> anchors = await log.GetAnchorsAsync(null, null, null, cancellationToken);
 
+				IEnumerable<ILogAnchor> matchingAnchors = anchors;
+				if (request.Severity != null)
+				{
+					matchingAnchors = anchors.Where(a => a.Severity == request.Severity.Value);
+				}
+
+				List<ILogAnchor> matchingList = matchingAnchors.ToList();
+				int total = matchingList.Count;
+
+				if (total == 0)
+				{
+					continue;
+				}
+
+				IEnumerable<ILogAnchor> anchorsToFetch = request.Count != null ? matchingList.Take(request.Count.Value) : matchingList;
+
 				Dictionary<ObjectId, int?> spanIdToIssueId = new Dictionary<ObjectId, int?>();
 				List<GetLogEventResponse> events = new List<GetLogEventResponse>();
 
-				foreach (ILogAnchor anchor in anchors)
+				foreach (ILogAnchor anchor in anchorsToFetch)
 				{
-					if (request.Severity != null && anchor.Severity != request.Severity.Value)
-					{
-						continue;
-					}
-
 					ILogEventData logEventData = await anchor.GetDataAsync(cancellationToken);
 
 					int? issueId = null;
@@ -423,17 +434,9 @@ namespace HordeServer.Logs
 					}
 
 					events.Add(CreateGetLogEventResponse(anchor, logEventData, issueId));
-
-					if (request.Count != null && events.Count >= request.Count.Value)
-					{
-						break;
-					}
 				}
 
-				if (events.Count > 0)
-				{
-					response.LogEvents[logIdStr] = events;
-				}
+				response.LogEvents[logIdStr] = new BatchLogEventsEntry { Total = total, Events = events };
 			}
 
 			return response;
diff --git a/Engine/Source/Programs/Shared/EpicGames.Horde/Logs/LogEventMessages.cs b/Engine/Source/Programs/Shared/EpicGames.Horde/Logs/LogEventMessages.cs
--- a/Engine/Source/Programs/Shared/EpicGames.Horde/Logs/LogEventMessages.cs
+++ b/Engine/Source/Programs/Shared/EpicGames.Horde/Logs/LogEventMessages.cs
@@ -54,6 +54,22 @@ namespace EpicGames.Horde.Logs
 		public LogEventSeverity? Severity { get; set; }
 	}
 
+	/// <summary>
+	/// Events for a single log within a batch response
+	/// </summary>
+	public class BatchLogEventsEntry
+	{
+		/// <summary>
+		/// Total number of matching events in this log (before count limit)
+		/// </summary>
+		public int Total { get; set; }
+
+		/// <summary>
+		/// The returned events (may be truncated by count)
+		/// </summary>
+		public List<GetLogEventResponse> Events { get; set; } = [];
+	}
+
 	/// <summary>
 	/// Response containing events grouped by log id
 	/// </summary>
@@ -62,7 +78,7 @@ namespace EpicGames.Horde.Logs
 		/// <summary>
 		/// Events grouped by log id
 		/// </summary>
-		public Dictionary<string, List<GetLogEventResponse>> LogEvents { get; set; } = [];
+		public Dictionary<string, BatchLogEventsEntry> LogEvents { get; set; } = [];
 	}
 
 	/// <summary>
