Horde: Fix breadcrumbs on tools view

diff --git a/Engine/Source/Programs/Horde/HordeDashboard/src/components/ToolView.tsx b/Engine/Source/Programs/Horde/HordeDashboard/src/components/ToolView.tsx
--- a/Engine/Source/Programs/Horde/HordeDashboard/src/components/ToolView.tsx
+++ b/Engine/Source/Programs/Horde/HordeDashboard/src/components/ToolView.tsx
@@ -392,7 +392,7 @@ export const ToolView: React.FC = () => {
 
    return <Stack className={hordeClasses.horde}>
       <TopNav />
-      <Breadcrumbs items={[{ text: 'Tools' }]} />
+      <Breadcrumbs items={[{ text: 'Tools' }, { text: 'Downloads' }]} />
       <ToolViewInner />
    </Stack>
 };
