[Test Health] Add support for csv data export

#rnx

#Synced-CL 51159424

diff --git a/Engine/Source/Programs/Horde/HordeDashboard/plugins/test-automation/testAutomationHealth.tsx b/Engine/Source/Programs/Horde/HordeDashboard/plugins/test-automation/testAutomationHealth.tsx
--- a/Engine/Source/Programs/Horde/HordeDashboard/plugins/test-automation/testAutomationHealth.tsx
+++ b/Engine/Source/Programs/Horde/HordeDashboard/plugins/test-automation/testAutomationHealth.tsx
@@ -1,6 +1,6 @@
 // Copyright Epic Games, Inc. All Rights Reserved.
 
-import { DetailsList, FontIcon, IColumn, IconButton, Label, Modal, SelectionMode, Stack, Text, HoverCard, HoverCardType, IPlainCardProps, DirectionalHint, DetailsListLayoutMode, Spinner, SpinnerSize, Toggle, DatePicker, ComboBox, IComboBoxOption, FocusZone, FocusZoneTabbableElements, FocusTrapCallout, PrimaryButton, DefaultButton, TooltipHost, TooltipOverflowMode, TooltipDelay, Link, IRenderFunction, IDetailsHeaderProps, Sticky, StickyPositionType, ScrollablePane } from "@fluentui/react";
+import { DetailsList, FontIcon, IColumn, IconButton, Label, Modal, SelectionMode, Stack, Text, HoverCard, HoverCardType, IPlainCardProps, DirectionalHint, DetailsListLayoutMode, Spinner, SpinnerSize, Toggle, DatePicker, ComboBox, IComboBoxOption, FocusZone, FocusZoneTabbableElements, FocusTrapCallout, PrimaryButton, DefaultButton, TooltipHost, TooltipOverflowMode, TooltipDelay, Link, IRenderFunction, IDetailsHeaderProps, Sticky, StickyPositionType, ScrollablePane, IContextualMenuProps } from "@fluentui/react";
 import { MetadataRef, TestDataHandler, TestNameRef, TestSessionResult } from "./testData";
 import { useCallback, useEffect, useMemo, useState } from "react";
 import { observer } from "mobx-react-lite";
@@ -15,6 +15,7 @@ import { booleanColors, cadenceColors, cadenceOrders, healthColors, healthStars,
 import { GetUserResponse, DeepLinkResponse } from "horde/backend/Api";
 import { errorDialogStore } from "horde/components/error/ErrorStore";
 import backend from "horde/backend";
+import { downloadCsv } from "horde/base/utilities/csvDownload";
 
 /// Colors
 const profileColors = new DomainNameColors();
@@ -542,6 +543,37 @@ type DeepLinkState = {
     };
 }
 
+const ToolsMenu: React.FC<{items: TestHealthItem[]}> = ({items}) => {
+
+    const options: IContextualMenuProps = {
+        items: [
+            {
+                key: 'export',
+                text: 'Export as CSV',
+                iconProps: {iconName: 'DownloadDocument'},
+                onClick: () => {
+                    downloadCsv(
+                        items.map(item => item.getAsArray(usersInfoCache)),
+                        'test-health.csv',
+                        TestHealthItem.getFields()
+                    );
+                }
+            },
+        ],
+    }
+
+    return  <IconButton
+                iconProps={{iconName: "ColumnOptions"}}
+                menuProps={options}
+                style={{
+                    fontSize:16, height: 28,
+                    border: '1px solid', borderColor: dashboard.darktheme ? "#4D4C4B" : "#6D6C6B", borderRadius: 4,
+                    backgroundColor: dashboard.darktheme ? "rgba(255, 255, 255, 0.1)" : "",
+                    color: dashboard.darktheme ? 'white' : 'grey'
+                }}
+            />
+}
+
 export const TestHealthView: React.FC<{ test: TestHealthItem, handler: TestDataHandler, onUpdateItem: (item: TestHealthItem) => void }> = observer(({ test, handler, onUpdateItem }) => {
     const [streams, setStreams] = useState<{current: string, toCompare?: string}>(
         {
@@ -1135,6 +1167,11 @@ export const TestAutomationHealth: React.FC<{ handler: TestDataHandler}> = obser
                         />
                     </ScrollablePane>
                 }
+                { !handler.queryAuditsLoading && 
+                    <Stack horizontalAlign='end' style={{ position: "absolute", right: 24, top: 8, zIndex: 1000 }} tokens={{childrenGap: 8}}>
+                        <ToolsMenu items={healthListState.items} />
+                    </Stack>
+                }
                 { handler.queryAuditsLoading &&
                     <Stack horizontalAlign='center' style={{ width: "100%", position: "absolute", top: 24 }} tokens={{childrenGap: 8}}>
                         <Text style={{ fontSize: 24 }}>Loading Data</Text>
diff --git a/Engine/Source/Programs/Horde/HordeDashboard/plugins/test-automation/testHealthData.ts b/Engine/Source/Programs/Horde/HordeDashboard/plugins/test-automation/testHealthData.ts
--- a/Engine/Source/Programs/Horde/HordeDashboard/plugins/test-automation/testHealthData.ts
+++ b/Engine/Source/Programs/Horde/HordeDashboard/plugins/test-automation/testHealthData.ts
@@ -337,6 +337,50 @@ export class TestHealthItem {
     get fetchingRecipes() {
         return this.handler.queryRecipesLoading;
     }
+
+    getAsArray(usersInfo?: UsersInfoCache): string[] {
+        return [
+            this.name,
+            this.stream,
+            this.health.toString(),
+            this.timings?.latestCompletedRun?.toISOString() ?? "",
+            this.cadence.toString(),
+            this.team ?? "",
+            this.profile ?? "",
+            this.harness ?? "",
+            this.intent ?? "",
+            this.notes ?? "",
+            usersInfo?.getInfo(this.owner)?.name ?? this.owner ?? "",
+            this.customers?.map(c => usersInfo?.getInfo(c)?.name ?? c).join(', ') ?? "",
+            this.notification ?? "",
+            this.underAudit.toString(),
+            this.lastAuditDate?.toISOString() ?? "",
+            usersInfo?.getInfo(this.user)?.name ?? this.user ?? "",
+            this.userInputDate?.toISOString() ?? "",
+        ];
+    }
+    
+    static getFields() {
+        return [
+            'Name',
+            'Stream',
+            'Health',
+            'Latest Completed Run',
+            'Cadence',
+            'Team',
+            'Profile',
+            'Harness',
+            'Intent',
+            'Notes',
+            'Owner',
+            'Customers',
+            'Notification',
+            'Under Audit',
+            'Last Audit Date',
+            'User',
+            'User Input Date',
+        ];
+    }
 }
 
 export class UsersInfoCache {
@@ -358,7 +402,7 @@ export class UsersInfoCache {
 
     getInfo(userId?: string) {
         return this.cache.get(userId);
-    } 
+    }
 
     private cache : Map<string | undefined, GetUserResponse>;
 }
diff --git a/Engine/Source/Programs/Horde/HordeDashboard/src/base/utilities/csvDownload.ts b/Engine/Source/Programs/Horde/HordeDashboard/src/base/utilities/csvDownload.ts
new file mode 100644
--- /dev/null
+++ b/Engine/Source/Programs/Horde/HordeDashboard/src/base/utilities/csvDownload.ts
@@ -0,0 +1,81 @@
+// Copyright Epic Games, Inc. All Rights Reserved.
+
+/**
+ * Escapes a CSV cell value (wraps in quotes if needed, escapes internal quotes).
+ */
+function escapeCsvCell(value: string): string {
+	const needsQuotes = /[",\r\n]/.test(value) || value.trim() !== value;
+	const escaped = value.replace(/"/g, '""');
+	return needsQuotes ? `"${escaped}"` : escaped;
+}
+
+/**
+ * Builds a CSV string from an array of rows. Each row is an array of cell values.
+ * The first row is treated as the header if headers are not provided separately.
+ */
+function buildCsv(rows: string[][], header?: string[]): string {
+	const bodyRows = header ? rows : rows.slice(1);
+	const firstRow = header ?? rows[0];
+	if (!firstRow?.length) {
+		return '';
+	}
+	const toLine = (row: string[]) => row.map(cell => escapeCsvCell(String(cell ?? ''))).join(',');
+	const lines = [toLine(firstRow), ...bodyRows.map(toLine)];
+	return lines.join('\n');
+}
+
+/**
+ * Triggers a file download in the browser using a Blob and a temporary anchor.
+ * Use this for client-generated content (e.g. CSV) so the user gets a named file.
+ */
+function downloadBlob(blob: Blob, filename: string): void {
+	const url = URL.createObjectURL(blob);
+	const link = document.createElement('a');
+	link.href = url;
+	link.download = filename;
+	link.style.display = 'none';
+	document.body.appendChild(link);
+	link.click();
+	document.body.removeChild(link);
+	URL.revokeObjectURL(url);
+}
+
+/**
+ * Generates a CSV file and triggers a download in the browser.
+ *
+ * @param rows - Array of rows; each row is an array of cell values (strings or numbers).
+ * @param filename - Name of the downloaded file (e.g. "export.csv").
+ * @param header - Optional. Pass a separate header row; otherwise the first row is used as header.
+ *
+ * @example
+ * // Header + data rows
+ * downloadCsv([
+ *   ['Name', 'Status', 'Date'],
+ *   ['Test A', 'Pass', '2025-02-17'],
+ *   ['Test B', 'Fail', '2025-02-17'],
+ * ], 'test-results.csv');
+ *
+ * @example
+ * // From array of objects (build rows yourself)
+ * const rows = items.map(i => [i.name, i.status, i.date]);
+ * downloadCsv(rows, 'export.csv', ['Name', 'Status', 'Date']);
+ */
+export function downloadCsv(rows: string[][], filename: string, header?: string[]): void {
+	const csv = buildCsv(rows, header);
+	// UTF-8 BOM so Excel and other tools recognize encoding
+	const bom = '\uFEFF';
+	const blob = new Blob([bom + csv], { type: 'text/csv;charset=utf-8' });
+	downloadBlob(blob, filename.endsWith('.csv') ? filename : `${filename}.csv`);
+}
+
+/**
+ * Triggers a download of a raw CSV string (e.g. if you built the CSV yourself).
+ *
+ * @param csvContent - Full CSV string (can include header line).
+ * @param filename - Name of the downloaded file (e.g. "export.csv").
+ */
+export function downloadCsvContent(csvContent: string, filename: string): void {
+	const bom = '\uFEFF';
+	const blob = new Blob([bom + csvContent], { type: 'text/csv;charset=utf-8' });
+	downloadBlob(blob, filename.endsWith('.csv') ? filename : `${filename}.csv`);
+}
