Horde: Updating schema documentation

#rnx

diff --git a/Engine/Source/Programs/Horde/Docs/Config/Schema/AclActions.md b/Engine/Source/Programs/Horde/Docs/Config/Schema/AclActions.md
--- a/Engine/Source/Programs/Horde/Docs/Config/Schema/AclActions.md
+++ b/Engine/Source/Programs/Horde/Docs/Config/Schema/AclActions.md
@@ -778,12 +778,6 @@
 | `UpdateReplicator` |  |
 | `ViewReplicator` |  |
 
-## Secrets
-
-| Name | Description |
-| ---- | ----------- |
-| `ViewSecret` |  |
-
 ## ServiceAccounts
 
 | Name | Description |
@@ -808,6 +802,19 @@
 | `DownloadSoftware` |  |
 | `DeleteSoftware` |  |
 
+## Storage
+
+| Name | Description |
+| ---- | ----------- |
+| `ReadBlobs` |  |
+| `WriteBlobs` |  |
+| `ReadAliases` |  |
+| `WriteAliases` |  |
+| `DeleteAliases` |  |
+| `ReadRefs` |  |
+| `WriteRefs` |  |
+| `DeleteRefs` |  |
+
 ## Streams
 
 | Name | Description |
@@ -955,12 +962,6 @@
 | `UpdateReplicator` |  |
 | `ViewReplicator` |  |
 
-## Secrets
-
-| Name | Description |
-| ---- | ----------- |
-| `ViewSecret` |  |
-
 ## ServiceAccounts
 
 | Name | Description |
@@ -1145,12 +1146,6 @@
 | `UpdateReplicator` |  |
 | `ViewReplicator` |  |
 
-## Secrets
-
-| Name | Description |
-| ---- | ----------- |
-| `ViewSecret` |  |
-
 ## ServiceAccounts
 
 | Name | Description |
@@ -1199,6 +1194,12 @@
 | `ViewChanges` |  |
 | `ViewTemplate` |  |
 
+## Symbols
+
+| Name | Description |
+| ---- | ----------- |
+| `ReadSymbols` |  |
+
 ## TestData
 
 | Name | Description |
@@ -1335,12 +1336,6 @@
 | `UpdateReplicator` |  |
 | `ViewReplicator` |  |
 
-## Secrets
-
-| Name | Description |
-| ---- | ----------- |
-| `ViewSecret` |  |
-
 ## ServiceAccounts
 
 | Name | Description |
@@ -1389,6 +1384,12 @@
 | `ViewChanges` |  |
 | `ViewTemplate` |  |
 
+## Symbols
+
+| Name | Description |
+| ---- | ----------- |
+| `ReadSymbols` |  |
+
 ## TestData
 
 | Name | Description |
diff --git a/Engine/Source/Programs/Horde/Docs/Config/Schema/Globals.md b/Engine/Source/Programs/Horde/Docs/Config/Schema/Globals.md
--- a/Engine/Source/Programs/Horde/Docs/Config/Schema/Globals.md
+++ b/Engine/Source/Programs/Horde/Docs/Config/Schema/Globals.md
@@ -57,13 +57,13 @@
 ---- | -----------
 `build` | [BuildConfig](#buildconfig)<br>Configuration for the build plugin
 `compute` | [ComputeConfig](#computeconfig)<br>Configuration for the compute plugin
-`exampleplugin` | [ExamplePluginConfig](#examplepluginconfig)<br>Configuration for the exampleplugin plugin
 `analytics` | [AnalyticsConfig](#analyticsconfig)<br>Configuration for the analytics plugin
-`secrets` | [SecretsConfig](#secretsconfig)<br>Configuration for the secrets plugin
+`exampleplugin` | [ExamplePluginConfig](#examplepluginconfig)<br>Configuration for the exampleplugin plugin
 `storage` | [StorageConfig](#storageconfig)<br>Configuration for the storage plugin
 `experimental` | [ExperimentalConfig](#experimentalconfig)<br>Configuration for the experimental plugin
-`tools` | [ToolsConfig](#toolsconfig)<br>Configuration for the tools plugin
 `symbols` | [SymbolsConfig](#symbolsconfig)<br>Configuration for the symbols plugin
+`tools` | [ToolsConfig](#toolsconfig)<br>Configuration for the tools plugin
+`secrets` | [SecretsConfig](#secretsconfig)<br>Configuration for the secrets plugin
 
 ## BuildConfig
 
@@ -381,83 +381,17 @@
 `description` | `string`<br>
 `computeId` | `string`<br>
 
-## ExamplePluginConfig
-
-Name | Description
----- | -----------
-`greetingMessage` | `string`<br>
-
 ## AnalyticsConfig
 
 Name | Description
 ---- | -----------
 `stores` | [TelemetryStoreConfig](Telemetry.md)`[]`<br>
 
-## SecretsConfig
+## ExamplePluginConfig
 
 Name | Description
 ---- | -----------
-`secrets` | [SecretConfig](#secretconfig)`[]`<br>
-`providerConfigs` | [SecretProviderConfig](#secretproviderconfig)`[]`<br>
+`greetingMessage` | `string`<br>
 
-## SecretConfig
-
-Name | Description
----- | -----------
-`id` | `string`<br>
-`data` | `string` `->` `string`<br>
-`sources` | [ExternalSecretConfig](#externalsecretconfig)`[]`<br>
-`acl` | [AclConfig](#aclconfig)<br>
-
-## ExternalSecretConfig
-
-Name | Description
----- | -----------
-`providerConfig` | `string`<br>
-`provider` | `string`<br>
-`format` | [ExternalSecretFormat](#externalsecretformat-enum)<br>
-`key` | `string`<br>
-`path` | `string`<br>
-
-## ExternalSecretFormat (Enum)
-
-Name | Description
----- | -----------
-`Text` | 
-`Json` | 
-
-## SecretProviderConfig
-
-Name | Description
----- | -----------
-`name` | `string`<br>
-`provider` | `string`<br>
-`hcpVault` | [HcpVaultConfig](#hcpvaultconfig)<br>
-`azureKeyVault` | [AzureKeyVaultConfig](#azurekeyvaultconfig)<br>
-
-## HcpVaultConfig
-
-Name | Description
----- | -----------
-`preSharedKey` | `string`<br>
-`endPoint` | `string`<br>
-`awsArnRole` | `string`<br>
-`awsIamServerId` | `string`<br>
-`role` | `string`<br>
-`credentials` | [HcpVaultCredentialsType](#hcpvaultcredentialstype-enum)<br>
-
-## HcpVaultCredentialsType (Enum)
-
-Name | Description
----- | -----------
-`PreSharedKey` | 
-`AwsAuth` | 
-
-## AzureKeyVaultConfig
-
-Name | Description
----- | -----------
-`vaultUri` | `string`<br>
-
 ## StorageConfig
 
@@ -527,6 +461,21 @@
 ---- | -----------
 `notifications` | [NotificationConfig](Notification.md)`[]`<br>
 
+## SymbolsConfig
+
+Name | Description
+---- | -----------
+`stores` | [SymbolStoreConfig](#symbolstoreconfig)`[]`<br>
+
+## SymbolStoreConfig
+
+Name | Description
+---- | -----------
+`id` | `string`<br>
+`namespaceId` | `string`<br>
+`public` | `boolean`<br>
+`acl` | [AclConfig](#aclconfig)<br>
+
 ## ToolsConfig
 
 Name | Description
@@ -551,17 +500,68 @@
 `namespaceId` | `string`<br>
 `acl` | [AclConfig](#aclconfig)<br>
 
-## SymbolsConfig
+## SecretsConfig
 
 Name | Description
 ---- | -----------
-`stores` | [SymbolStoreConfig](#symbolstoreconfig)`[]`<br>
+`secrets` | [SecretConfig](#secretconfig)`[]`<br>
+`providerConfigs` | [SecretProviderConfig](#secretproviderconfig)`[]`<br>
 
-## SymbolStoreConfig
+## SecretConfig
 
 Name | Description
 ---- | -----------
 `id` | `string`<br>
-`namespaceId` | `string`<br>
-`public` | `boolean`<br>
+`data` | `string` `->` `string`<br>
+`sources` | [ExternalSecretConfig](#externalsecretconfig)`[]`<br>
 `acl` | [AclConfig](#aclconfig)<br>
+
+## ExternalSecretConfig
+
+Name | Description
+---- | -----------
+`providerConfig` | `string`<br>
+`provider` | `string`<br>
+`format` | [ExternalSecretFormat](#externalsecretformat-enum)<br>
+`key` | `string`<br>
+`path` | `string`<br>
+
+## ExternalSecretFormat (Enum)
+
+Name | Description
+---- | -----------
+`Text` | 
+`Json` | 
+
+## SecretProviderConfig
+
+Name | Description
+---- | -----------
+`name` | `string`<br>
+`provider` | `string`<br>
+`hcpVault` | [HcpVaultConfig](#hcpvaultconfig)<br>
+`azureKeyVault` | [AzureKeyVaultConfig](#azurekeyvaultconfig)<br>
+
+## HcpVaultConfig
+
+Name | Description
+---- | -----------
+`preSharedKey` | `string`<br>
+`endPoint` | `string`<br>
+`awsArnRole` | `string`<br>
+`awsIamServerId` | `string`<br>
+`role` | `string`<br>
+`credentials` | [HcpVaultCredentialsType](#hcpvaultcredentialstype-enum)<br>
+
+## HcpVaultCredentialsType (Enum)
+
+Name | Description
+---- | -----------
+`PreSharedKey` | 
+`AwsAuth` | 
+
+## AzureKeyVaultConfig
+
+Name | Description
+---- | -----------
+`vaultUri` | `string`<br>
diff --git a/Engine/Source/Programs/Horde/Docs/Deployment/ServerSettings.md b/Engine/Source/Programs/Horde/Docs/Deployment/ServerSettings.md
--- a/Engine/Source/Programs/Horde/Docs/Deployment/ServerSettings.md
+++ b/Engine/Source/Programs/Horde/Docs/Deployment/ServerSettings.md
@@ -127,13 +127,13 @@ Name | Description
 ---- | -----------
 `build` | [BuildServerConfig](#buildserverconfig)<br>Configuration for the build plugin
 `compute` | [ComputeServerConfig](#computeserverconfig)<br>Configuration for the compute plugin
-`exampleplugin` | [ExamplePluginServerConfig](#examplepluginserverconfig)<br>Configuration for the exampleplugin plugin
 `analytics` | [AnalyticsServerConfig](#analyticsserverconfig)<br>Configuration for the analytics plugin
-`secrets` | [SecretsServerConfig](#secretsserverconfig)<br>Configuration for the secrets plugin
+`exampleplugin` | [ExamplePluginServerConfig](#examplepluginserverconfig)<br>Configuration for the exampleplugin plugin
 `storage` | [StorageServerConfig](#storageserverconfig)<br>Configuration for the storage plugin
 `experimental` | [ExperimentalServerConfig](#experimentalserverconfig)<br>Configuration for the experimental plugin
-`tools` | [ToolsServerConfig](#toolsserverconfig)<br>Configuration for the tools plugin
 `symbols` | [PluginServerConfig](#pluginserverconfig)<br>Configuration for the symbols plugin
+`tools` | [ToolsServerConfig](#toolsserverconfig)<br>Configuration for the tools plugin
+`secrets` | [SecretsServerConfig](#secretsserverconfig)<br>Configuration for the secrets plugin
 
 ## BuildServerConfig
 
@@ -310,12 +310,6 @@ Name | Description
 `ComputeQueueAwsMetric` | 
 `LeaseUtilizationAwsMetric` | 
 
-## ExamplePluginServerConfig
-
-Name | Description
----- | -----------
-`enabled` | `boolean`<br>
-
 ## AnalyticsServerConfig
 
 Name | Description
@@ -345,11 +339,10 @@ Name | Description
 `retainDays` | `number`<br>
 `enabled` | `boolean`<br>
 
-## SecretsServerConfig
+## ExamplePluginServerConfig
 
 Name | Description
 ---- | -----------
-`withAws` | `boolean`<br>
 `enabled` | `boolean`<br>
 
 ## StorageServerConfig
@@ -409,6 +402,12 @@ Name | Description
 ---- | -----------
 `enabled` | `boolean`<br>
 
+## PluginServerConfig
+
+Name | Description
+---- | -----------
+`enabled` | `boolean`<br>
+
 ## ToolsServerConfig
 
 Name | Description
@@ -470,8 +469,9 @@ Name | Description
 `excludeActions` | `string[]`<br>
 `extends` | `string[]`<br>
 
-## PluginServerConfig
+## SecretsServerConfig
 
 Name | Description
 ---- | -----------
+`withAws` | `boolean`<br>
 `enabled` | `boolean`<br>
