Horde: Tune the step events for job step history

diff --git a/Engine/Source/Programs/Horde/HordeDashboard/src/components/jobDetailsV2/JobDetailStepHistory.tsx b/Engine/Source/Programs/Horde/HordeDashboard/src/components/jobDetailsV2/JobDetailStepHistory.tsx
--- a/Engine/Source/Programs/Horde/HordeDashboard/src/components/jobDetailsV2/JobDetailStepHistory.tsx
+++ b/Engine/Source/Programs/Horde/HordeDashboard/src/components/jobDetailsV2/JobDetailStepHistory.tsx
@@ -21,7 +21,7 @@ import { getStepRefMessageInfo } from "./JobDetailsUtils";
 const sideRail: ISideRailLink = { text: "History", url: "rail_step_history" };
 
 /** Maximum number of recent failing/warning steps to fetch log events for */
-const MAX_STEPS_FOR_EVENTS = 30;
+const MAX_STEPS_FOR_EVENTS = 15;
 
 class StepHistoryDataView extends JobDataView {
 
