Horde: Optimize FindIssuesV2 to use aggregate for distinct issue IDs and fetch spans only for paginated results

diff --git a/Engine/Source/Programs/Horde/Plugins/Build/HordeServer.Build/Issues/IIssueCollection.cs b/Engine/Source/Programs/Horde/Plugins/Build/HordeServer.Build/Issues/IIssueCollection.cs
--- a/Engine/Source/Programs/Horde/Plugins/Build/HordeServer.Build/Issues/IIssueCollection.cs
+++ b/Engine/Source/Programs/Horde/Plugins/Build/HordeServer.Build/Issues/IIssueCollection.cs
@@ -182,6 +182,17 @@ namespace HordeServer.Issues
 		/// <returns>List of streams open in the given stream at the given changelist</returns>
 		Task<IReadOnlyList<IIssueSpan>> FindSpansAsync(IEnumerable<ObjectId>? ids = null, IEnumerable<int>? issueIds = null, StreamId? streamId = null, CommitId? minCommitId = null, CommitId? maxCommitId = null, bool? resolved = null, int? index = null, int? count = null, CancellationToken cancellationToken = default);
 
+		/// <summary>
+		/// Gets distinct issue ids from spans matching the given filter, using an aggregate pipeline to avoid transferring full span documents.
+		/// </summary>
+		/// <param name="streamId">The stream to query</param>
+		/// <param name="minCommitId">Minimum commit id filter</param>
+		/// <param name="maxCommitId">Maximum commit id filter</param>
+		/// <param name="resolved">Include resolved or unresolved spans</param>
+		/// <param name="cancellationToken">Cancellation token for the operation</param>
+		/// <returns>List of distinct issue ids</returns>
+		Task<IReadOnlyList<int>> FindSpanIssueIdsAsync(StreamId streamId, CommitId? minCommitId = null, CommitId? maxCommitId = null, bool? resolved = null, CancellationToken cancellationToken = default);
+
 		#endregion
 
 		#region Steps
diff --git a/Engine/Source/Programs/Horde/Plugins/Build/HordeServer.Build/Issues/IssueCollection.cs b/Engine/Source/Programs/Horde/Plugins/Build/HordeServer.Build/Issues/IssueCollection.cs
--- a/Engine/Source/Programs/Horde/Plugins/Build/HordeServer.Build/Issues/IssueCollection.cs
+++ b/Engine/Source/Programs/Horde/Plugins/Build/HordeServer.Build/Issues/IssueCollection.cs
@@ -739,6 +739,8 @@ namespace HordeServer.Issues
 			List<MongoIndex<IssueSpanDocument>> issueSpanIndexes = new List<MongoIndex<IssueSpanDocument>>();
 			issueSpanIndexes.Add(keys => keys.Ascending(x => x.IssueId));
 			issueSpanIndexes.Add(keys => keys.Ascending(x => x.StreamId).Ascending(x => x.MinCommitOrder).Ascending(x => x.MaxCommitOrder));
+			// This index is created, though we need to move away from these C# index declarations, so only documenting it here
+			// issueSpanIndexes.Add(keys => keys.Ascending(x => x.StreamId).Ascending(x => x.ResolvedAt));
 			issueSpanIndexes.Add("StreamChanges", keys => keys.Ascending(x => x.StreamId).Ascending(x => x.TemplateRefId).Ascending(x => x.NodeName).Ascending(x => x.MinCommitOrder).Ascending(x => x.MaxCommitOrder));
 			_issueSpans = mongoService.GetCollection<IssueSpanDocument>("IssuesV2.Spans", issueSpanIndexes);
 
@@ -1885,6 +1887,38 @@ namespace HordeServer.Issues
 			return await _issueSpans.Find(filter).Range(index, count).ToListAsync(cancellationToken);
 		}
 
+		/// <inheritdoc/>
+		public async Task<IReadOnlyList<int>> FindSpanIssueIdsAsync(StreamId streamId, CommitId? minCommitId, CommitId? maxCommitId, bool? resolved, CancellationToken cancellationToken)
+		{
+			FilterDefinition<IssueSpanDocument> filter = Builders<IssueSpanDocument>.Filter.Eq(x => x.StreamId, streamId);
+
+			if (minCommitId != null)
+			{
+				CommitIdWithOrder minCommitIdWithOrder = await _commitService.GetOrderedAsync(streamId, minCommitId, cancellationToken);
+				filter &= Builders<IssueSpanDocument>.Filter.Not(Builders<IssueSpanDocument>.Filter.Lt(x => x.MaxCommitOrder, minCommitIdWithOrder.Order));
+			}
+			if (maxCommitId != null)
+			{
+				CommitIdWithOrder maxCommitIdWithOrder = await _commitService.GetOrderedAsync(streamId, maxCommitId, cancellationToken);
+				filter &= Builders<IssueSpanDocument>.Filter.Not(Builders<IssueSpanDocument>.Filter.Gt(x => x.MinCommitOrder, maxCommitIdWithOrder.Order));
+			}
+
+			if (resolved != null)
+			{
+				if (resolved.Value)
+				{
+					filter &= Builders<IssueSpanDocument>.Filter.Ne(x => x.ResolvedAt, null);
+				}
+				else
+				{
+					filter &= Builders<IssueSpanDocument>.Filter.Eq(x => x.ResolvedAt, null);
+				}
+			}
+
+			List<ProjectedIssueId> result = await _issueSpans.Aggregate().Match(filter).Group(x => x.IssueId, x => new ProjectedIssueId { _id = x.Key }).ToListAsync(cancellationToken);
+			return result.Where(x => x._id != null).Select(x => x._id!.Value).ToList();
+		}
+
 		#endregion
 
 		#region Steps
diff --git a/Engine/Source/Programs/Horde/Plugins/Build/HordeServer.Build/Issues/IssuesController.cs b/Engine/Source/Programs/Horde/Plugins/Build/HordeServer.Build/Issues/IssuesController.cs
--- a/Engine/Source/Programs/Horde/Plugins/Build/HordeServer.Build/Issues/IssuesController.cs
+++ b/Engine/Source/Programs/Horde/Plugins/Build/HordeServer.Build/Issues/IssuesController.cs
@@ -94,15 +94,27 @@ namespace HordeServer.Issues
 					return Forbid(StreamAclAction.ViewStream, streamId.Value);
 				}
 
-				IReadOnlyList<IIssueSpan> spans = await _issueCollection.FindSpansAsync(null, ids, streamId.Value, minChange, maxChange, resolved, cancellationToken: cancellationToken);
-				if (spans.Count > 0)
+				// Get distinct issue ids from spans using aggregate (avoids transferring full span documents)
+				IReadOnlyList<int> issueIds = await _issueCollection.FindSpanIssueIdsAsync(streamId.Value, minChange, maxChange, resolved, cancellationToken);
+				if (ids != null)
 				{
-					// Group all the spans by their issue id
+					HashSet<int> spanIssueIds = new HashSet<int>(issueIds);
+					issueIds = ids.Where(x => spanIssueIds.Contains(x)).ToArray();
+				}
+
+				if (issueIds.Count > 0)
+				{
+					// Paginate the issues
+					IReadOnlyList<IIssue> issues = await _issueCollection.FindIssuesAsync(issueIds, index: index, count: count, cancellationToken: cancellationToken);
+
+					// Fetch spans only for the paginated issues
+					IReadOnlyList<IIssueSpan> spans = await _issueCollection.FindSpansAsync(null, issues.Select(x => x.Id), streamId.Value, minChange, maxChange, resolved, cancellationToken: cancellationToken);
+
+					// Group spans by issue id
 					Dictionary<int, List<IIssueSpan>> issueIdToSpans = new Dictionary<int, List<IIssueSpan>>();
 					foreach (IIssueSpan span in spans)
 					{
-						List<IIssueSpan>? spansForIssue;
-						if (!issueIdToSpans.TryGetValue(span.IssueId, out spansForIssue))
+						if (!issueIdToSpans.TryGetValue(span.IssueId, out List<IIssueSpan>? spansForIssue))
 						{
 							spansForIssue = new List<IIssueSpan>();
 							issueIdToSpans.Add(span.IssueId, spansForIssue);
@@ -110,9 +122,6 @@ namespace HordeServer.Issues
 						spansForIssue.Add(span);
 					}
 
-					// Find the matching issues
-					IReadOnlyList<IIssue> issues = await _issueCollection.FindIssuesAsync(issueIdToSpans.Keys, index: index, count: count, cancellationToken: cancellationToken);
-
 					// Create the corresponding responses
 					foreach (IIssue issue in issues.OrderByDescending(x => x.Id))
 					{
