Horde: Limit size of log messages in server

Too large log messages can cause issues with external log handling pipelines (containerd/FluentBit/OTel etc). This change prevents too deep destructuring of exceptions, a common source of over-sized messages.
Also inserts a JSON formatter for Serilog that caps log entries by size, replacing them with a compact JSON stub with the bare minimum information.

diff --git a/Engine/Source/Programs/Horde/HordeServer.Tests.Shared/ServerTestSetup.cs b/Engine/Source/Programs/Horde/HordeServer.Tests.Shared/ServerTestSetup.cs
--- a/Engine/Source/Programs/Horde/HordeServer.Tests.Shared/ServerTestSetup.cs
+++ b/Engine/Source/Programs/Horde/HordeServer.Tests.Shared/ServerTestSetup.cs
@@ -68,7 +68,8 @@ namespace HordeServer.Tests
 				.Enrich.FromLogContext()
 				.Enrich.WithExceptionDetails(new DestructuringOptionsBuilder()
 					.WithDefaultDestructurers()
-					.WithDestructurers(new[] { new RpcExceptionDestructurer() }))
+					.WithDestructurers([new RpcExceptionDestructurer()])
+					.WithDestructuringDepth(3))
 				.WriteTo.Console(outputTemplate: "[{Timestamp:HH:mm:ss} {Level:w3}] {Indent}{Message:l}{NewLine}{Exception}")
 				.CreateLogger();
 		}
diff --git a/Engine/Source/Programs/Horde/HordeServer.Tests/Utilities/SizeLimitedJsonFormatterTests.cs b/Engine/Source/Programs/Horde/HordeServer.Tests/Utilities/SizeLimitedJsonFormatterTests.cs
new file mode 100644
--- /dev/null
+++ b/Engine/Source/Programs/Horde/HordeServer.Tests/Utilities/SizeLimitedJsonFormatterTests.cs
@@ -0,0 +1,83 @@
+// Copyright Epic Games, Inc. All Rights Reserved.
+
+using System;
+using System.IO;
+using HordeServer.Utilities;
+using Microsoft.VisualStudio.TestTools.UnitTesting;
+using Serilog.Events;
+using Serilog.Formatting.Json;
+using Serilog.Parsing;
+
+namespace HordeServer.Tests.Utilities
+{
+	[TestClass]
+	public class SizeLimitedJsonFormatterTests
+	{
+		readonly JsonFormatter _inner = new JsonFormatter(renderMessage: true);
+		readonly SizeLimitedJsonFormatter _truncating = new SizeLimitedJsonFormatter(new JsonFormatter(renderMessage: true), maxBytes: 100);
+
+		static LogEvent CreateLogEvent(string messageTemplate, Exception? exception = null, LogEventLevel level = LogEventLevel.Information)
+		{
+			MessageTemplateParser parser = new MessageTemplateParser();
+			return new LogEvent(DateTimeOffset.UtcNow, level, exception, parser.Parse(messageTemplate), Array.Empty<LogEventProperty>());
+		}
+
+		static string Format(Serilog.Formatting.ITextFormatter formatter, LogEvent logEvent)
+		{
+			using StringWriter writer = new StringWriter();
+			formatter.Format(logEvent, writer);
+			return writer.ToString();
+		}
+
+		[TestMethod]
+		public void SmallEventPassesThroughUnchanged()
+		{
+			LogEvent logEvent = CreateLogEvent("Hello world");
+			SizeLimitedJsonFormatter formatter = new SizeLimitedJsonFormatter(_inner, maxBytes: 256 * 1024);
+
+			Assert.AreEqual(Format(_inner, logEvent), Format(formatter, logEvent));
+		}
+
+		[TestMethod]
+		public void OversizedEventIsTruncated()
+		{
+			LogEvent logEvent = CreateLogEvent("This message template is long enough to exceed the tiny 100-byte limit we set for testing purposes");
+			string result = Format(_truncating, logEvent);
+
+			Assert.IsTrue(result.Contains("\"_truncated\":true", StringComparison.Ordinal), "Should contain _truncated marker");
+			Assert.IsTrue(result.Contains("\"Level\":\"Information\"", StringComparison.Ordinal), "Should contain level");
+			Assert.IsTrue(result.Contains("\"MessageTemplate\":", StringComparison.Ordinal), "Should contain message template");
+			Assert.IsFalse(result.Contains("\"RenderedMessage\"", StringComparison.Ordinal), "Should not contain full rendered output");
+			Assert.IsFalse(result.Contains("ExceptionType", StringComparison.Ordinal), "Should not contain exception fields when no exception");
+			Assert.IsFalse(result.Contains("ExceptionMessage", StringComparison.Ordinal), "Should not contain exception fields when no exception");
+		}
+
+		[TestMethod]
+		public void TruncatedEventWithExceptionIncludesExceptionInfo()
+		{
+			InvalidOperationException exception = new InvalidOperationException("Something went wrong");
+			LogEvent logEvent = CreateLogEvent("Operation failed", exception, LogEventLevel.Error);
+			string result = Format(_truncating, logEvent);
+
+			Assert.IsTrue(result.Contains("\"_truncated\":true", StringComparison.Ordinal), "Should contain _truncated marker");
+			Assert.IsTrue(result.Contains("\"ExceptionType\":\"System.InvalidOperationException\"", StringComparison.Ordinal), "Should contain exception type");
+			Assert.IsTrue(result.Contains("\"ExceptionMessage\":\"Something went wrong\"", StringComparison.Ordinal), "Should contain exception message");
+		}
+
+		[TestMethod]
+		public void TruncatedEventIsValidJson()
+		{
+			Exception exception = new Exception("Bad \"things\"\nhappened");
+			LogEvent logEvent = CreateLogEvent("Template with \"quotes\" and\nnewlines", exception, LogEventLevel.Error);
+			string result = Format(_truncating, logEvent).TrimEnd();
+
+			// Verify it parses as valid JSON
+			System.Text.Json.JsonDocument.Parse(result);
+
+			// Verify special characters were escaped
+			Assert.IsFalse(result.Contains("\n\"", StringComparison.Ordinal), "Newlines in values should be escaped");
+			Assert.IsTrue(result.Contains("\\n", StringComparison.Ordinal), "Should contain escaped newlines");
+			Assert.IsTrue(result.Contains("\\\"quotes\\\"", StringComparison.Ordinal), "Should contain escaped quotes");
+		}
+	}
+}
diff --git a/Engine/Source/Programs/Horde/HordeServer/ServerApp.cs b/Engine/Source/Programs/Horde/HordeServer/ServerApp.cs
--- a/Engine/Source/Programs/Horde/HordeServer/ServerApp.cs
+++ b/Engine/Source/Programs/Horde/HordeServer/ServerApp.cs
@@ -48,7 +48,7 @@ namespace HordeServer
 		{
 			if (settings.LogJsonToStdOut)
 			{
-				return sinkConfig.Console(new JsonFormatter(renderMessage: true));
+				return sinkConfig.Console(new SizeLimitedJsonFormatter(new JsonFormatter(renderMessage: true)));
 			}
 			else
 			{
@@ -185,10 +185,11 @@ namespace HordeServer
 				.Enrich.FromLogContext()
 				.Enrich.WithExceptionDetails(new DestructuringOptionsBuilder()
 					.WithDefaultDestructurers()
-					.WithDestructurers(new[] { new RpcExceptionDestructurer() }))
+					.WithDestructurers([new RpcExceptionDestructurer()])
+					.WithDestructuringDepth(3))
 				.WriteTo.Console(serverSettings)
 				.WriteTo.File(Path.Combine(logDir.FullName, "Log.txt"), outputTemplate: "[{Timestamp:HH:mm:ss} {Level:w3}] {Indent}{Message:l}{NewLine}{Exception}", rollingInterval: RollingInterval.Day, rollOnFileSizeLimit: true, fileSizeLimitBytes: 20 * 1024 * 1024, retainedFileCountLimit: 10)
-				.WriteTo.File(new JsonFormatter(renderMessage: true), Path.Combine(logDir.FullName, "Log.json"), rollingInterval: RollingInterval.Day, rollOnFileSizeLimit: true, fileSizeLimitBytes: 20 * 1024 * 1024, retainedFileCountLimit: 10)
+				.WriteTo.File(new SizeLimitedJsonFormatter(new JsonFormatter(renderMessage: true)), Path.Combine(logDir.FullName, "Log.json"), rollingInterval: RollingInterval.Day, rollOnFileSizeLimit: true, fileSizeLimitBytes: 20 * 1024 * 1024, retainedFileCountLimit: 10)
 				.ReadFrom.Configuration(config)
 				.CreateLogger();
 
diff --git a/Engine/Source/Programs/Horde/HordeServer/Utilities/SizeLimitedJsonFormatter.cs b/Engine/Source/Programs/Horde/HordeServer/Utilities/SizeLimitedJsonFormatter.cs
new file mode 100644
--- /dev/null
+++ b/Engine/Source/Programs/Horde/HordeServer/Utilities/SizeLimitedJsonFormatter.cs
@@ -0,0 +1,69 @@
+// Copyright Epic Games, Inc. All Rights Reserved.
+
+using System.IO;
+using System.Text.Encodings.Web;
+using System.Text.Json;
+using Serilog.Events;
+using Serilog.Formatting;
+
+namespace HordeServer.Utilities
+{
+	/// <summary>
+	/// Wraps an <see cref="ITextFormatter"/> and enforces a maximum byte size on the formatted output.
+	/// If the inner formatter produces output exceeding the limit, a truncated replacement is emitted
+	/// containing the timestamp, level, message template, exception type/message, and a _truncated marker.
+	/// </summary>
+	sealed class SizeLimitedJsonFormatter : ITextFormatter
+	{
+		static readonly JsonSerializerOptions s_jsonOptions = new JsonSerializerOptions { Encoder = JavaScriptEncoder.UnsafeRelaxedJsonEscaping };
+
+		readonly ITextFormatter _inner;
+		readonly int _maxBytes;
+
+		/// <summary>
+		/// Constructor
+		/// </summary>
+		/// <param name="inner">The inner formatter to delegate to</param>
+		/// <param name="maxBytes">Maximum byte size before truncation (default 256 KB)</param>
+		public SizeLimitedJsonFormatter(ITextFormatter inner, int maxBytes = 256 * 1024)
+		{
+			_inner = inner;
+			_maxBytes = maxBytes;
+		}
+
+		/// <inheritdoc/>
+		public void Format(LogEvent logEvent, TextWriter output)
+		{
+			using StringWriter buffer = new StringWriter();
+			_inner.Format(logEvent, buffer);
+
+			string formatted = buffer.ToString();
+
+			int byteCount = System.Text.Encoding.UTF8.GetByteCount(formatted);
+			if (byteCount <= _maxBytes)
+			{
+				output.Write(formatted);
+				return;
+			}
+
+			// Emit a truncated replacement that is still valid JSON for Datadog
+			string timestamp = logEvent.Timestamp.UtcDateTime.ToString("O");
+			string level = logEvent.Level.ToString();
+			string messageTemplate = EscapeJson(logEvent.MessageTemplate.Text);
+			string exceptionType = logEvent.Exception?.GetType().FullName ?? "";
+			string exceptionMessage = EscapeJson(logEvent.Exception?.Message ?? "");
+
+			output.Write($"{{\"Timestamp\":\"{timestamp}\",\"Level\":\"{level}\",\"MessageTemplate\":\"{messageTemplate}\"");
+			if (logEvent.Exception != null)
+			{
+				output.Write($",\"ExceptionType\":\"{EscapeJson(exceptionType)}\",\"ExceptionMessage\":\"{exceptionMessage}\"");
+			}
+			output.Write($",\"_truncated\":true,\"_originalBytes\":{byteCount}}}\n");
+		}
+
+		static string EscapeJson(string value)
+		{
+			return JsonSerializer.Serialize(value, s_jsonOptions)[1..^1];
+		}
+	}
+}
