Horde: Change conform check when assigning a job task to have up to date free disk space for the agent

On lease completion disk free space property of the agent is updated to become available for the next lease assignment.
Implements a minimal ISystemMetrics for macOS to support disk usage.

#Synced-CL 0

diff --git a/Engine/Source/Programs/Horde/Drivers/JobDriver.Tests/Execution/ConformExecutorTests.cs b/Engine/Source/Programs/Horde/Drivers/JobDriver.Tests/Execution/ConformExecutorTests.cs
--- a/Engine/Source/Programs/Horde/Drivers/JobDriver.Tests/Execution/ConformExecutorTests.cs
+++ b/Engine/Source/Programs/Horde/Drivers/JobDriver.Tests/Execution/ConformExecutorTests.cs
@@ -2,10 +2,8 @@
 
 using System;
 using System.Collections.Generic;
-using System.IO;
 using System.Threading;
 using System.Threading.Tasks;
-using EpicGames.Core;
 using EpicGames.Horde.Agents;
 using EpicGames.Horde.Agents.Leases;
 using EpicGames.Perforce.Fixture;
@@ -14,7 +12,6 @@ using HordeCommon.Rpc.Tasks;
 using JobDriver.Execution;
 using Microsoft.Extensions.DependencyInjection;
 using Microsoft.Extensions.Logging;
-using Microsoft.Extensions.Logging.Abstractions;
 using Microsoft.VisualStudio.TestTools.UnitTesting;
 using OpenTelemetry.Trace;
 using SyncOptions = JobDriver.Execution.SyncOptions;
@@ -102,46 +99,6 @@ public class ConformTest : BasePerforceFixtureTest
 		FakeHordeClient hordeClient = new (jobRpcClientStub);
 		return new ConformExecutor(hordeClient, _materializerFactories, TempDir, new AgentId("agent1"), leaseId, conformTask,
 			new DriverSettings(),
-			new FakeSystemMetrics(),
 			NoOpTracer, _logger);
 	}
 }
-
-[TestClass]
-public class ConformNoPerforceTest
-{
-	private static Tracer NoOpTracer { get; } = TracerProvider.Default.GetTracer("NoOp");
-
-	[TestMethod]
-	[DataRow(500, true, DisplayName = "Conform skipped because there is enough disk free space available")]
-	[DataRow(5000, false, DisplayName = "Conform not skipped because there is not enough disk free space available")]
-	[DataRow(0, false, DisplayName = "Conform not skipped because ConformDiskFreeSpace workspace property is not set")]
-	public async Task SkipConformWhenAvailableFreeDiskSpaceAsync(long conformDiskFreeSpace, bool skippedResult)
-	{
-		ConformTask conformTask = new()
-		{
-			SkipWhenDiskFreeSpaceIsAvailable = true,
-			Workspaces = { new List<RpcAgentWorkspace> { new() { ConformDiskFreeSpace = conformDiskFreeSpace } } }
-		};
-
-		// Required because not using BasePerforceFixtureTest and must not run p4 commands;
-		// to avoid it using the local machine's p4 connection settings
-		DriverSettings driverSettings = new() { Executor = String.Empty, PerforceExecutor = { RunConform = false } };
-
-		ConformExecutor executor = new(
-			new FakeHordeClient(new JobRpcClientStub(NullLogger.Instance)),
-			[],
-			new DirectoryReference(Path.GetTempPath()),
-			new AgentId("agent1"),
-			LeaseId.Parse("111111111111111111111111"),
-			conformTask,
-			driverSettings,
-			new FakeSystemMetrics { FreeSpace = 1000 },
-			NoOpTracer,
-			NullLogger.Instance);
-
-		await executor.ExecuteAsync(CancellationToken.None);
-
-		Assert.AreEqual(skippedResult, executor.ConformWasSkipped);
-	}
-}
diff --git a/Engine/Source/Programs/Horde/Drivers/JobDriver/Commands/Execution/ExecuteConformCommand.cs b/Engine/Source/Programs/Horde/Drivers/JobDriver/Commands/Execution/ExecuteConformCommand.cs
--- a/Engine/Source/Programs/Horde/Drivers/JobDriver/Commands/Execution/ExecuteConformCommand.cs
+++ b/Engine/Source/Programs/Horde/Drivers/JobDriver/Commands/Execution/ExecuteConformCommand.cs
@@ -4,7 +4,6 @@ using EpicGames.Core;
 using EpicGames.Horde;
 using EpicGames.Horde.Agents;
 using EpicGames.Horde.Agents.Leases;
-using EpicGames.Horde.Agents.Metrics;
 using EpicGames.Horde.Agents.Sessions;
 using JobDriver.Execution;
 using HordeCommon.Rpc.Tasks;
@@ -35,15 +34,13 @@ namespace JobDriver.Commands.Execution
 		readonly HordeClientFactory _hordeClientFactory;
 		readonly IEnumerable<IWorkspaceMaterializerFactory> _materializerFactories;
 		readonly IOptions<DriverSettings> _driverSettings;
-		readonly ISystemMetrics _systemMetrics;
 		readonly Tracer _tracer;
 
-		public ExecuteConformCommand(HordeClientFactory hordeClientFactory, IEnumerable<IWorkspaceMaterializerFactory> materializerFactories, IOptions<DriverSettings> driverSettings, ISystemMetrics systemMetrics, Tracer tracer)
+		public ExecuteConformCommand(HordeClientFactory hordeClientFactory, IEnumerable<IWorkspaceMaterializerFactory> materializerFactories, IOptions<DriverSettings> driverSettings, Tracer tracer)
 		{
 			_hordeClientFactory = hordeClientFactory;
 			_materializerFactories = materializerFactories;
 			_driverSettings = driverSettings;
-			_systemMetrics = systemMetrics;
 			_tracer = tracer;
 		}
 
@@ -55,7 +52,7 @@ namespace JobDriver.Commands.Execution
 
 			await using HordeClient hordeClient = _hordeClientFactory.Create();
 
-			ConformExecutor conformExecutor = new ConformExecutor(hordeClient, _materializerFactories, WorkingDir, AgentId, LeaseId, conformTask, _driverSettings.Value, _systemMetrics, _tracer, logger);
+			ConformExecutor conformExecutor = new ConformExecutor(hordeClient, _materializerFactories, WorkingDir, AgentId, LeaseId, conformTask, _driverSettings.Value, _tracer, logger);
 			await conformExecutor.ExecuteAsync(CancellationToken.None);
 
 			return 0;
diff --git a/Engine/Source/Programs/Horde/Drivers/JobDriver/DriverApp.cs b/Engine/Source/Programs/Horde/Drivers/JobDriver/DriverApp.cs
--- a/Engine/Source/Programs/Horde/Drivers/JobDriver/DriverApp.cs
+++ b/Engine/Source/Programs/Horde/Drivers/JobDriver/DriverApp.cs
@@ -3,13 +3,11 @@
 using System.Reflection;
 using EpicGames.Core;
 using EpicGames.Horde;
-using EpicGames.Horde.Agents.Metrics;
 using EpicGames.Horde.Utilities;
 using JobDriver.Execution;
 using JobDriver.Utility;
 using Microsoft.Extensions.Configuration;
 using Microsoft.Extensions.DependencyInjection;
-using Microsoft.Extensions.Logging;
 using OpenTelemetry.Trace;
 
 namespace JobDriver
@@ -36,7 +34,7 @@ namespace JobDriver
 
 			// Create the services
 			IServiceCollection services = new ServiceCollection();
-			RegisterServices(services, arguments);
+			RegisterServices(services);
 
 			// Run the host
 			await using ServiceProvider serviceProvider = services.BuildServiceProvider();
@@ -54,16 +52,6 @@ namespace JobDriver
 		/// </summary>
 		/// <param name="services"></param>
 		public static void RegisterServices(IServiceCollection services)
-		{
-			RegisterServices(services, new CommandLineArguments([]));
-		}
-
-		/// <summary>
-		/// Helper method to register services for this app
-		/// </summary>
-		/// <param name="services"></param>
-		/// <param name="arguments"></param>
-		public static void RegisterServices(IServiceCollection services, CommandLineArguments arguments)
 		{
 			// Read the driver config
 			IConfiguration configuration = new ConfigurationBuilder()
@@ -103,20 +91,7 @@ namespace JobDriver
 			services.AddSingleton<IWorkspaceMaterializerFactory, ManagedWorkspaceMaterializerFactory>();
 			services.AddSingleton<IWorkspaceMaterializerFactory, PerforceMaterializerFactory>();
 
-			arguments.TryGetValue("-WorkingDir=", out string? workingDir);
-			services.AddSingleton<ISystemMetrics>(sp => CreateSystemMetrics(workingDir, sp.GetRequiredService<ILogger<ISystemMetrics>>()));
-
 			services.AddCommandsFromAssembly(Assembly.GetExecutingAssembly());
 		}
-
-		static ISystemMetrics CreateSystemMetrics(string? workingDir, ILogger logger)
-		{
-			DirectoryReference? wd = null;
-			if (workingDir != null)
-			{
-				wd = new DirectoryReference(workingDir);
-			}
-			return new DefaultSystemMetrics(wd, logger);
-		}
 	}
 }
diff --git a/Engine/Source/Programs/Horde/Drivers/JobDriver/Execution/ConformExecutor.cs b/Engine/Source/Programs/Horde/Drivers/JobDriver/Execution/ConformExecutor.cs
--- a/Engine/Source/Programs/Horde/Drivers/JobDriver/Execution/ConformExecutor.cs
+++ b/Engine/Source/Programs/Horde/Drivers/JobDriver/Execution/ConformExecutor.cs
@@ -4,7 +4,6 @@ using EpicGames.Core;
 using EpicGames.Horde;
 using EpicGames.Horde.Agents;
 using EpicGames.Horde.Agents.Leases;
-using EpicGames.Horde.Agents.Metrics;
 using JobDriver.Utility;
 using Horde.Common.Rpc;
 using HordeCommon.Rpc.Messages;
@@ -23,20 +22,11 @@ namespace JobDriver.Execution
 		readonly IEnumerable<IWorkspaceMaterializerFactory> _materializerFactories;
 		readonly ConformTask _conformTask;
 		readonly DriverSettings _driverSettings;
-		readonly ISystemMetrics _systemMetrics;
 		readonly Tracer _tracer;
 		readonly ILogger _logger;
 
-		/// <summary>
-		/// Whether it was decided to skip the conform during execution.
-		/// </summary>
-		/// <remarks>
-		/// Exposed to ease testing.
-		/// </remarks>
-		internal bool ConformWasSkipped { get; private set; }
-
 		public ConformExecutor(IHordeClient hordeClient, IEnumerable<IWorkspaceMaterializerFactory> materializerFactories,
-			DirectoryReference workingDir, AgentId agentId, LeaseId leaseId, ConformTask conformTask, DriverSettings driverSettings, ISystemMetrics systemMetrics, Tracer tracer, ILogger logger)
+			DirectoryReference workingDir, AgentId agentId, LeaseId leaseId, ConformTask conformTask, DriverSettings driverSettings, Tracer tracer, ILogger logger)
 		{
 			_hordeClient = hordeClient;
 			_workingDir = workingDir;
@@ -45,7 +35,6 @@ namespace JobDriver.Execution
 			_materializerFactories = materializerFactories;
 			_conformTask = conformTask;
 			_driverSettings = driverSettings;
-			_systemMetrics = systemMetrics;
 			_tracer = tracer;
 			_logger = logger;
 		}
@@ -114,12 +103,7 @@ namespace JobDriver.Execution
 					.Select(x => DirectoryReference.Combine(_workingDir, x.Identifier)).ToList();
 
 				// Run the conform task
-				if (_conformTask.SkipWhenDiskFreeSpaceIsAvailable && HasAvailableDiskFreeSpace(pendingWorkspaces))
-				{
-					ConformWasSkipped = true;
-					_logger.LogInformation("Skipping conform. Executor={Executor} RunConform={RunConform}", _driverSettings.Executor, _driverSettings.PerforceExecutor.RunConform);
-				}
-				else if (isExecutorConformCompatible && _driverSettings.PerforceExecutor.RunConform)
+				if (isExecutorConformCompatible && _driverSettings.PerforceExecutor.RunConform)
 				{
 					if (managedWorkspaces.Count > 0)
 					{
@@ -157,33 +141,5 @@ namespace JobDriver.Execution
 				removeUntrackedFiles = response.RemoveUntrackedFiles;
 			}
 		}
-
-		bool HasAvailableDiskFreeSpace(IList<RpcAgentWorkspace> pendingWorkspaces)
-		{
-			long? freeDiskSpace = _systemMetrics.GetWorkingDisk()?.FreeSpace;
-			if (freeDiskSpace == null)
-			{
-				return false;
-			}
-
-			long conformDiskSpaceNeeded = 0;
-			foreach (RpcAgentWorkspace workspace in pendingWorkspaces)
-			{
-				if (workspace.ConformDiskFreeSpace > 0)
-				{
-					conformDiskSpaceNeeded = Math.Max(conformDiskSpaceNeeded, workspace.ConformDiskFreeSpace);
-				}
-			}
-
-			bool hasAvailableFreeDiskSpace = conformDiskSpaceNeeded > 0 && freeDiskSpace >= conformDiskSpaceNeeded;
-			if (hasAvailableFreeDiskSpace)
-			{
-				double BytesToMegabytes(long v) => v / 1024.0 / 1024.0;
-				_logger.LogInformation("Skipping conform as free disk space available ({FreeDiskSpace:F1} MB) is greater than the required free disk space when the conform was requested ({RequiredFreeDiskSpace:F1} MB)",
-					BytesToMegabytes(freeDiskSpace.Value), BytesToMegabytes(conformDiskSpaceNeeded));
-			}
-
-			return hasAvailableFreeDiskSpace;
-		}
 	}
 }
diff --git a/Engine/Source/Programs/Horde/HordeAgent.Tests/Leases/LeaseManagerTests.cs b/Engine/Source/Programs/Horde/HordeAgent.Tests/Leases/LeaseManagerTests.cs
--- a/Engine/Source/Programs/Horde/HordeAgent.Tests/Leases/LeaseManagerTests.cs
+++ b/Engine/Source/Programs/Horde/HordeAgent.Tests/Leases/LeaseManagerTests.cs
@@ -11,6 +11,7 @@ using EpicGames.Core;
 using EpicGames.Horde;
 using EpicGames.Horde.Agents;
 using EpicGames.Horde.Agents.Leases;
+using EpicGames.Horde.Agents.Metrics;
 using EpicGames.Horde.Agents.Sessions;
 using HordeAgent.Leases;
 using HordeAgent.Services;
@@ -34,6 +35,16 @@ public class FakeCapabilitiesService : ICapabilitiesService
 	}
 }
 
+public class FakeSystemMetrics : ISystemMetrics
+{
+	public long DiskFreeSpace { get; set; } = 1000;
+	public CpuMetrics? GetCpu() => null;
+	public CpuTimeMetrics? GetCpuTime() => null;
+	public MemoryMetrics? GetMemory() => null;
+	public DiskMetrics? GetWorkingDisk() => new() { FreeSpace = DiskFreeSpace, TotalSize = 10000 };
+	public DiskMetrics? GetSystemDisk() => null;
+}
+
 public class SessionStub(AgentId agentId, SessionId sessionId, DirectoryReference workingDir, IHordeClient hordeClient) : ISession
 {
 	public AgentId AgentId { get; } = agentId;
@@ -140,7 +151,63 @@ public sealed class LeaseManagerTests
 		Assert.AreEqual(new SessionResult(SessionOutcome.Terminate, SessionReason.Completed), await _leaseManager.RunAsync(cts.Token));
 		Assert.AreEqual(RpcAgentStatus.Stopped, _hordeServer.AgentStatus);
 	}
-	
+
+	[TestMethod]
+	public async Task UpdateSession_SendsDiskFreeSpaceInCapabilitiesAsync()
+	{
+		using CancellationTokenSource cts = new(5000);
+		FakeSystemMetrics fakeMetrics = new() { DiskFreeSpace = 2000 };
+		FakeHordeRpcServer server = new(CreateConsoleLogger<FakeHordeRpcServer>());
+		LeaseManager leaseManager = CreateLeaseManager(server.GetHordeClient(), systemMetrics: fakeMetrics);
+		leaseManager.TerminateSessionAfterLease = true;
+
+		await leaseManager.RunAsync(cts.Token);
+		server.UpdateSessionRequests.Writer.Complete();
+
+		List<RpcUpdateSessionRequest> requests = await server.UpdateSessionRequests.Reader.ReadAllAsync(cts.Token).ToListAsync(cts.Token);
+		RpcUpdateSessionRequest request = requests.Last();
+		Assert.Contains($"{KnownPropertyNames.WorkingDiskFreeSpace}=2000", request.CapabilityOverrides.Properties, StringComparer.Ordinal);
+	}
+
+	[TestMethod]
+	public async Task UpdateSession_SendsDiskFreeSpaceInCapabilityOverridesAsync()
+	{
+		using CancellationTokenSource cts = new(5000);
+		FakeSystemMetrics fakeMetrics = new() { DiskFreeSpace = 2000 };
+		FakeCapabilitiesService capsService = new()
+		{
+			Capabilities = new RpcAgentCapabilities { Properties = { $"{KnownPropertyNames.WorkingDiskFreeSpace}=1000" } }
+		};
+
+		FakeHordeRpcServer server = new(CreateConsoleLogger<FakeHordeRpcServer>());
+		LeaseManager leaseManager = CreateLeaseManager(server.GetHordeClient(), capsService, systemMetrics: fakeMetrics);
+		leaseManager.TerminateSessionAfterLease = true;
+
+		await leaseManager.RunAsync(cts.Token);
+		server.UpdateSessionRequests.Writer.Complete();
+
+		List<RpcUpdateSessionRequest> requests = await server.UpdateSessionRequests.Reader.ReadAllAsync(cts.Token).ToListAsync(cts.Token);
+		RpcUpdateSessionRequest request = requests.Last();
+		Assert.Contains($"{KnownPropertyNames.WorkingDiskFreeSpace}=1000", request.Capabilities.Properties, StringComparer.Ordinal);
+		Assert.Contains($"{KnownPropertyNames.WorkingDiskFreeSpace}=2000", request.CapabilityOverrides.Properties, StringComparer.Ordinal);
+	}
+
+	[TestMethod]
+	public async Task UpdateSession_NoDiskMetrics_NoCapabilityOverridesAsync()
+	{
+		using CancellationTokenSource cts = new(5000);
+		FakeHordeRpcServer server = new(CreateConsoleLogger<FakeHordeRpcServer>());
+		LeaseManager leaseManager = CreateLeaseManager(server.GetHordeClient());
+		leaseManager.TerminateSessionAfterLease = true;
+
+		await leaseManager.RunAsync(cts.Token);
+		server.UpdateSessionRequests.Writer.Complete();
+
+		List<RpcUpdateSessionRequest> requests = await server.UpdateSessionRequests.Reader.ReadAllAsync(cts.Token).ToListAsync(cts.Token);
+		RpcUpdateSessionRequest request = requests.Last();
+		Assert.IsNull(request.CapabilityOverrides);
+	}
+
 	/// <summary>
 	/// Create a console logger for tests
 	/// </summary>
@@ -158,7 +225,7 @@ public sealed class LeaseManagerTests
 	}
 	
 	[SuppressMessage("Reliability", "CA2000:Dispose objects before losing scope")]
-	private static LeaseManager CreateLeaseManager(IHordeClient hordeClient, FakeCapabilitiesService? capsService = null, LeaseResult? leaseResult = null)
+	private static LeaseManager CreateLeaseManager(IHordeClient hordeClient, FakeCapabilitiesService? capsService = null, LeaseResult? leaseResult = null, ISystemMetrics? systemMetrics = null)
 	{
 		DirectoryReference tempWorkingDir = new(Path.Join(Path.GetTempPath(), Path.GetRandomFileName()));
 		AgentSettings settings = new() { WriteStepOutputToLogger = true };
@@ -170,7 +237,7 @@ public sealed class LeaseManagerTests
 			session,
 			capsService ?? new FakeCapabilitiesService(),
 			statusService,
-			new DefaultSystemMetrics(),
+			systemMetrics ?? new DefaultSystemMetrics(),
 			new List<LeaseHandlerFactory>() { new TestHandlerFactory(leaseResult) },
 			new LeaseLoggerFactory(settingsOptions, CreateConsoleLogger<LeaseLoggerFactory>()),
 			settingsOptions,
diff --git a/Engine/Source/Programs/Horde/HordeAgent/AgentApp.cs b/Engine/Source/Programs/Horde/HordeAgent/AgentApp.cs
--- a/Engine/Source/Programs/Horde/HordeAgent/AgentApp.cs
+++ b/Engine/Source/Programs/Horde/HordeAgent/AgentApp.cs
@@ -254,6 +254,10 @@ namespace HordeAgent
 				{
 					return new LinuxSystemMetrics(workingDir);
 				}
+				if (RuntimeInformation.IsOSPlatform(OSPlatform.OSX))
+				{
+					return new MacSystemMetrics(workingDir);
+				}
 			}
 			catch (Exception e)
 			{
diff --git a/Engine/Source/Programs/Horde/HordeAgent/Leases/LeaseManager.cs b/Engine/Source/Programs/Horde/HordeAgent/Leases/LeaseManager.cs
--- a/Engine/Source/Programs/Horde/HordeAgent/Leases/LeaseManager.cs
+++ b/Engine/Source/Programs/Horde/HordeAgent/Leases/LeaseManager.cs
@@ -3,6 +3,7 @@
 using System.Diagnostics;
 using System.Text;
 using EpicGames.Core;
+using EpicGames.Horde.Agents;
 using EpicGames.Horde.Agents.Leases;
 using EpicGames.Horde.Agents.Metrics;
 using Google.Protobuf.WellKnownTypes;
@@ -338,6 +339,14 @@ namespace HordeAgent.Leases
 				// Update the capabilities whenever the background task has generated a new instance
 				updateSessionRequest.Capabilities = Interlocked.Exchange(ref _capabilities, null);
 
+				// Add working disk free space as a capability override (updated every UpdateSession for fresh data)
+				DiskMetrics? workingDiskMetrics = _systemMetrics.GetWorkingDisk();
+				if (workingDiskMetrics != null)
+				{
+					updateSessionRequest.CapabilityOverrides ??= new RpcAgentCapabilities();
+					updateSessionRequest.CapabilityOverrides.Properties.Add($"{KnownPropertyNames.WorkingDiskFreeSpace}={workingDiskMetrics.FreeSpace}");
+				}
+
 				// Complete the wait task if we subsequently stop
 				using (stopping ? (CancellationTokenRegistration?)null : stoppingToken.Register(() => _updateLeasesEvent.Set()))
 				{
diff --git a/Engine/Source/Programs/Horde/HordeAgent/Services/TelemetryService.cs b/Engine/Source/Programs/Horde/HordeAgent/Services/TelemetryService.cs
--- a/Engine/Source/Programs/Horde/HordeAgent/Services/TelemetryService.cs
+++ b/Engine/Source/Programs/Horde/HordeAgent/Services/TelemetryService.cs
@@ -370,6 +370,95 @@ public sealed class LinuxSystemMetrics : ISystemMetrics, IDisposable
 	}
 }
 
+/// <summary>
+/// macOS specific implementation for gathering system metrics
+/// </summary>
+[SupportedOSPlatform("macos")]
+public sealed class MacSystemMetrics : ISystemMetrics
+{
+	private readonly DirectoryReference _workingDir;
+	private readonly Stopwatch _driveInfoUpdateTimer = Stopwatch.StartNew();
+	private DriveInfo? _workingDriveInfo;
+	private DriveInfo? _systemDriveInfo;
+
+	/// <summary>
+	/// Constructor
+	/// </summary>
+	public MacSystemMetrics(DirectoryReference workingDir)
+	{
+		_workingDir = workingDir;
+		_workingDriveInfo = GetDriveInfo(EpicGames.Horde.Agents.Metrics.DriveType.Working);
+		_systemDriveInfo = GetDriveInfo(EpicGames.Horde.Agents.Metrics.DriveType.System);
+	}
+
+	/// <inheritdoc/>
+	public CpuMetrics? GetCpu() => null;
+
+	/// <inheritdoc/>
+	public CpuTimeMetrics? GetCpuTime() => null;
+
+	/// <inheritdoc/>
+	public MemoryMetrics? GetMemory() => null;
+
+	/// <inheritdoc/>
+	public DiskMetrics? GetWorkingDisk()
+	{
+		RefreshDriveInfo();
+
+		if (_workingDriveInfo == null)
+		{
+			return null;
+		}
+
+		return new DiskMetrics
+		{
+			FreeSpace = _workingDriveInfo.AvailableFreeSpace,
+			TotalSize = _workingDriveInfo.TotalSize
+		};
+	}
+
+	/// <inheritdoc/>
+	public DiskMetrics? GetSystemDisk()
+	{
+		RefreshDriveInfo();
+
+		if (_systemDriveInfo == null)
+		{
+			return null;
+		}
+
+		return new DiskMetrics
+		{
+			FreeSpace = _systemDriveInfo.AvailableFreeSpace,
+			TotalSize = _systemDriveInfo.TotalSize
+		};
+	}
+
+	/// <summary>
+	/// Restart the drive info poll timer if the time period has elapsed
+	/// </summary>
+	void RefreshDriveInfo()
+	{
+		if (_driveInfoUpdateTimer.Elapsed > TimeSpan.FromSeconds(20))
+		{
+			_workingDriveInfo = GetDriveInfo(EpicGames.Horde.Agents.Metrics.DriveType.Working);
+			_systemDriveInfo = GetDriveInfo(EpicGames.Horde.Agents.Metrics.DriveType.System);
+			_driveInfoUpdateTimer.Restart();
+		}
+	}
+
+	/// <summary>
+	/// Fetches the drive info for the given drive type
+	/// </summary>
+	/// <param name="type">Classification of the drive we're fetching telemetry for</param>
+	/// <returns>DriveInfo for the matching drive</returns>
+	private DriveInfo GetDriveInfo(EpicGames.Horde.Agents.Metrics.DriveType type)
+	{
+		string driveName = type == EpicGames.Horde.Agents.Metrics.DriveType.Working ? _workingDir.FullName : "/";
+		return new DriveInfo(driveName);
+	}
+}
+
 /// <summary>
 /// Send telemetry events back to server at regular intervals
 /// </summary>
diff --git a/Engine/Source/Programs/Horde/Plugins/Build/HordeServer.Build.Tests/Agents/AgentControllerTest.cs b/Engine/Source/Programs/Horde/Plugins/Build/HordeServer.Build.Tests/Agents/AgentControllerTest.cs
--- a/Engine/Source/Programs/Horde/Plugins/Build/HordeServer.Build.Tests/Agents/AgentControllerTest.cs
+++ b/Engine/Source/Programs/Horde/Plugins/Build/HordeServer.Build.Tests/Agents/AgentControllerTest.cs
@@ -29,116 +29,32 @@ public class AgentControllerDbTest : BuildTestSetup
 	}
 
 	[TestMethod]
-	public async Task UpdateAgentFullConformWithDiskSpaceNeededAsync()
+	public async Task UpdateAgentFullConformAsync()
 	{
 		Fixture fixture = await CreateFixtureAsync();
 		IAgent agent = fixture.Agent1;
 
 		GetAgentResponse response = Deref((await AgentsController.GetAgentAsync(agent.Id)).Value as GetAgentResponse);
 		Assert.IsFalse(response.PendingFullConform);
-		Assert.IsFalse(response.PendingConformSkipWhenDiskFreeSpaceIsAvailable);
 
-		await AgentsController.UpdateAgentAsync(agent.Id, new UpdateAgentRequest(RequestFullConform: true, RequestConformSkipWhenAvailableFreeDiskSpace: true));
-
-		response = Deref((await AgentsController.GetAgentAsync(agent.Id)).Value as GetAgentResponse);
-		Assert.IsTrue(response.PendingFullConform);
-		Assert.IsTrue(response.PendingConformSkipWhenDiskFreeSpaceIsAvailable);
-	}
-
-	[TestMethod]
-	public async Task UpdateAgentFullConformWithoutDiskSpaceNeededAsync()
-	{
-		Fixture fixture = await CreateFixtureAsync();
-		IAgent agent = fixture.Agent1;
-
-		GetAgentResponse response = Deref((await AgentsController.GetAgentAsync(agent.Id)).Value as GetAgentResponse);
-		Assert.IsFalse(response.PendingFullConform);
-		Assert.IsFalse(response.PendingConformSkipWhenDiskFreeSpaceIsAvailable);
-
-		await AgentsController.UpdateAgentAsync(agent.Id, new UpdateAgentRequest(RequestFullConform: true));
-
-		response = Deref((await AgentsController.GetAgentAsync(agent.Id)).Value as GetAgentResponse);
-		Assert.IsTrue(response.PendingFullConform);
-		Assert.IsFalse(response.PendingConformSkipWhenDiskFreeSpaceIsAvailable);
-	}
-
-	[TestMethod]
-	public async Task UpdateAgentIgnoreDiskSpaceNeededWhenNotFullConformAsync()
-	{
-		Fixture fixture = await CreateFixtureAsync();
-		IAgent agent = fixture.Agent1;
-
-		GetAgentResponse response = Deref((await AgentsController.GetAgentAsync(agent.Id)).Value as GetAgentResponse);
-		Assert.IsFalse(response.PendingFullConform);
-		Assert.IsFalse(response.PendingConform);
-		Assert.IsFalse(response.PendingConformSkipWhenDiskFreeSpaceIsAvailable);
-
-		// check disk space option is ignored when a conform is not requested
-		await AgentsController.UpdateAgentAsync(agent.Id, new UpdateAgentRequest(RequestConformSkipWhenAvailableFreeDiskSpace: true));
-
-		response = Deref((await AgentsController.GetAgentAsync(agent.Id)).Value as GetAgentResponse);
-		Assert.IsFalse(response.PendingFullConform);
-		Assert.IsFalse(response.PendingConform);
-		Assert.IsFalse(response.PendingConformSkipWhenDiskFreeSpaceIsAvailable);
-
-		// When not a full conform check disk space option should be ignored because untracked files won't be deleted
-		await AgentsController.UpdateAgentAsync(agent.Id, new UpdateAgentRequest(RequestConform: true, RequestConformSkipWhenAvailableFreeDiskSpace: true));
-
-		response = Deref((await AgentsController.GetAgentAsync(agent.Id)).Value as GetAgentResponse);
-		Assert.IsTrue(response.PendingConform);
-		Assert.IsFalse(response.PendingConformSkipWhenDiskFreeSpaceIsAvailable);
-	}
-
-	[TestMethod]
-	public async Task UpdateAgentSuccessiveFullConformRemovePreviousDiskSpaceNeededAsync()
-	{
-		Fixture fixture = await CreateFixtureAsync();
-		IAgent agent = fixture.Agent1;
-
-		GetAgentResponse response = Deref((await AgentsController.GetAgentAsync(agent.Id)).Value as GetAgentResponse);
-		Assert.IsFalse(response.PendingFullConform);
-		Assert.IsFalse(response.PendingConformSkipWhenDiskFreeSpaceIsAvailable);
-
-		await AgentsController.UpdateAgentAsync(agent.Id, new UpdateAgentRequest(RequestFullConform: true, RequestConformSkipWhenAvailableFreeDiskSpace: true));
-
-		response = Deref((await AgentsController.GetAgentAsync(agent.Id)).Value as GetAgentResponse);
-		Assert.IsTrue(response.PendingFullConform);
-		Assert.IsTrue(response.PendingConformSkipWhenDiskFreeSpaceIsAvailable);
-
-		// Successive full conform request without an explicit check disk space option will remove it
 		await AgentsController.UpdateAgentAsync(agent.Id, new UpdateAgentRequest(RequestFullConform: true));
 
 		response = Deref((await AgentsController.GetAgentAsync(agent.Id)).Value as GetAgentResponse);
-		Assert.IsNotNull(response);
 		Assert.IsTrue(response.PendingFullConform);
-		Assert.IsFalse(response.PendingConformSkipWhenDiskFreeSpaceIsAvailable);
 	}
 
 	[TestMethod]
-	public async Task UpdateAgentSuccessiveConformRemovePreviousDiskSpaceNeededAsync()
+	public async Task UpdateAgentConformAsync()
 	{
 		Fixture fixture = await CreateFixtureAsync();
 		IAgent agent = fixture.Agent1;
 
 		GetAgentResponse response = Deref((await AgentsController.GetAgentAsync(agent.Id)).Value as GetAgentResponse);
-		Assert.IsFalse(response.PendingFullConform);
 		Assert.IsFalse(response.PendingConform);
-		Assert.IsFalse(response.PendingConformSkipWhenDiskFreeSpaceIsAvailable);
 
-		await AgentsController.UpdateAgentAsync(agent.Id, new UpdateAgentRequest(RequestFullConform: true, RequestConformSkipWhenAvailableFreeDiskSpace: true));
-
-		response = Deref((await AgentsController.GetAgentAsync(agent.Id)).Value as GetAgentResponse);
-		Assert.IsTrue(response.PendingFullConform);
-		Assert.IsTrue(response.PendingConformSkipWhenDiskFreeSpaceIsAvailable);
-
-		// Successive conform request will remove the check disk space option
 		await AgentsController.UpdateAgentAsync(agent.Id, new UpdateAgentRequest(RequestConform: true));
 
 		response = Deref((await AgentsController.GetAgentAsync(agent.Id)).Value as GetAgentResponse);
-		Assert.IsNotNull(response);
 		Assert.IsTrue(response.PendingConform);
-		// The previous full conform remains but check disk space needed is removed by the new conform request
-		Assert.IsTrue(response.PendingFullConform);
-		Assert.IsFalse(response.PendingConformSkipWhenDiskFreeSpaceIsAvailable);
 	}
-}
\ No newline at end of file
+}
diff --git a/Engine/Source/Programs/Horde/Plugins/Build/HordeServer.Build.Tests/Agents/AgentServiceTest.cs b/Engine/Source/Programs/Horde/Plugins/Build/HordeServer.Build.Tests/Agents/AgentServiceTest.cs
--- a/Engine/Source/Programs/Horde/Plugins/Build/HordeServer.Build.Tests/Agents/AgentServiceTest.cs
+++ b/Engine/Source/Programs/Horde/Plugins/Build/HordeServer.Build.Tests/Agents/AgentServiceTest.cs
@@ -253,14 +253,12 @@ public class AgentServiceTest : BuildTestSetup
 		AgentId agentId = new(fixture.Agent1Name);
 
 		IAgent agent = Deref(await AgentService.GetAgentAsync(agentId));
-		agent = Deref(await agent.TryUpdateAsync(new UpdateAgentOptions(RequestFullConform: true, RequestConformSkipWhenDiskFreeSpaceIsAvailable: true)));
+		agent = Deref(await agent.TryUpdateAsync(new UpdateAgentOptions(RequestFullConform: true)));
 		Assert.IsTrue(agent.RequestFullConform);
-		Assert.IsTrue(agent.RequestConformSkipWhenDiskFreeSpaceIsAvailable);
 
 		await AgentService.TryUpdateWorkspacesAsync(agent, agent.Workspaces.ToList(), false, CancellationToken.None);
 
 		agent = Deref(await AgentService.GetAgentAsync(agentId));
 		Assert.IsFalse(agent.RequestFullConform);
-		Assert.IsFalse(agent.RequestConformSkipWhenDiskFreeSpaceIsAvailable);
 	}
 }
\ No newline at end of file
diff --git a/Engine/Source/Programs/Horde/Plugins/Build/HordeServer.Build.Tests/Agents/Pools/PoolUpdateServiceTest.cs b/Engine/Source/Programs/Horde/Plugins/Build/HordeServer.Build.Tests/Agents/Pools/PoolUpdateServiceTest.cs
--- a/Engine/Source/Programs/Horde/Plugins/Build/HordeServer.Build.Tests/Agents/Pools/PoolUpdateServiceTest.cs
+++ b/Engine/Source/Programs/Horde/Plugins/Build/HordeServer.Build.Tests/Agents/Pools/PoolUpdateServiceTest.cs
@@ -130,7 +130,6 @@ public class PoolUpdateServiceTest : BuildTestSetup
 
 		// Assert
 		Assert.AreEqual(conformRequested, agent.RequestFullConform);
-		Assert.IsFalse(agent.RequestConformSkipWhenDiskFreeSpaceIsAvailable);
 	}
 
 	private async Task<IAgent> CreateAutoConformAgentAsync(int freeDiskSpaceMb, params int?[] autoConformThresholdsMb)
diff --git a/Engine/Source/Programs/Horde/Plugins/Build/HordeServer.Build.Tests/Jobs/JobTaskSourceTests.cs b/Engine/Source/Programs/Horde/Plugins/Build/HordeServer.Build.Tests/Jobs/JobTaskSourceTests.cs
--- a/Engine/Source/Programs/Horde/Plugins/Build/HordeServer.Build.Tests/Jobs/JobTaskSourceTests.cs
+++ b/Engine/Source/Programs/Horde/Plugins/Build/HordeServer.Build.Tests/Jobs/JobTaskSourceTests.cs
@@ -209,13 +209,11 @@ namespace HordeServer.Tests.Jobs
 			if (expectingConform)
 			{
 				Assert.IsTrue(agent.RequestFullConform);
-				Assert.IsTrue(agent.RequestConformSkipWhenDiskFreeSpaceIsAvailable);
 				Assert.IsNull(leaseResult);
 			}
 			else
 			{
 				Assert.IsFalse(agent.RequestFullConform);
-				Assert.IsFalse(agent.RequestConformSkipWhenDiskFreeSpaceIsAvailable);
 				Assert.IsNotNull(leaseResult);
 			}
 		}
@@ -228,15 +226,9 @@ namespace HordeServer.Tests.Jobs
 			Fixture fixture = await CreateFixtureAsync();
 			IAgent agent = fixture.Agent1;
 
-			UpdateAgentOptions options = new() { RequestConformSkipWhenDiskFreeSpaceIsAvailable = true };
-			if (fullConform)
-			{
-				options = options with { RequestFullConform = true };
-			}
-			else
-			{
-				options = options with { RequestConform = true };
-			}
+			UpdateAgentOptions options = fullConform
+				? new() { RequestFullConform = true }
+				: new() { RequestConform = true };
 
 			agent = Deref(await agent.TryUpdateAsync(options));
 
@@ -248,12 +240,10 @@ namespace HordeServer.Tests.Jobs
 
 			if (fullConform)
 			{
-				Assert.IsTrue(conformTask.SkipWhenDiskFreeSpaceIsAvailable);
 				Assert.IsTrue(conformTask.RemoveUntrackedFiles);
 			}
 			else
 			{
-				Assert.IsFalse(conformTask.SkipWhenDiskFreeSpaceIsAvailable);
 				Assert.IsFalse(conformTask.RemoveUntrackedFiles);
 			}
 		}
diff --git a/Engine/Source/Programs/Horde/Plugins/Build/HordeServer.Build/Agents/Pools/PoolUpdateService.cs b/Engine/Source/Programs/Horde/Plugins/Build/HordeServer.Build/Agents/Pools/PoolUpdateService.cs
--- a/Engine/Source/Programs/Horde/Plugins/Build/HordeServer.Build/Agents/Pools/PoolUpdateService.cs
+++ b/Engine/Source/Programs/Horde/Plugins/Build/HordeServer.Build/Agents/Pools/PoolUpdateService.cs
@@ -140,12 +140,9 @@ namespace HordeServer.Agents.Pools
 			List<IAgent> agents = (await _agentService.GetCachedAgentsAsync(cancellationToken))
 				.Where(x => x is { Status: AgentStatus.Stopped, Deleted: false, Enabled: true, Ephemeral: false }).ToList();
 
-			// To keep the existing behavior for the scheduled conforms, agent updates will not skip the conform
-			bool skipWhenDiskFreeSpaceIsAvailable = false;
-
 			foreach (IAgent agent in agents)
 			{
-				await ConformUtilities.AutoConformAgentAsync(agent, _poolService, _buildConfig.CurrentValue, skipWhenDiskFreeSpaceIsAvailable, _logger, cancellationToken);
+				await ConformUtilities.AutoConformAgentAsync(agent, _poolService, _buildConfig.CurrentValue, _logger, cancellationToken);
 			}
 		}
 
diff --git a/Engine/Source/Programs/Horde/Plugins/Build/HordeServer.Build/Jobs/JobTaskSource.cs b/Engine/Source/Programs/Horde/Plugins/Build/HordeServer.Build/Jobs/JobTaskSource.cs
--- a/Engine/Source/Programs/Horde/Plugins/Build/HordeServer.Build/Jobs/JobTaskSource.cs
+++ b/Engine/Source/Programs/Horde/Plugins/Build/HordeServer.Build/Jobs/JobTaskSource.cs
@@ -785,12 +785,7 @@ namespace HordeServer.Jobs
 				return true;
 			}
 
-			// When a conform is scheduled over a job recheck the available disk free space on lease execution.
-			// Agent capabilities are updated asynchronously to lease assignment this means when the conform is executed
-			// the available disk free space on the agent may now be above the conform threshold, which allows the conform to be skipped.
-			bool skipWhenDiskFreeSpaceIsAvailable = true;
-
-			return await ConformUtilities.AutoConformAgentAsync(agent, _poolService, _buildConfig.CurrentValue, skipWhenDiskFreeSpaceIsAvailable, _logger, cancellationToken);
+			return await ConformUtilities.AutoConformAgentAsync(agent, _poolService, _buildConfig.CurrentValue, _logger, cancellationToken);
 		}
 
 		/// <summary>
diff --git a/Engine/Source/Programs/Horde/Plugins/Build/HordeServer.Build/Tasks/ConformTaskSource.cs b/Engine/Source/Programs/Horde/Plugins/Build/HordeServer.Build/Tasks/ConformTaskSource.cs
--- a/Engine/Source/Programs/Horde/Plugins/Build/HordeServer.Build/Tasks/ConformTaskSource.cs
+++ b/Engine/Source/Programs/Horde/Plugins/Build/HordeServer.Build/Tasks/ConformTaskSource.cs
@@ -174,7 +174,6 @@ namespace HordeServer.Tasks
 							ILog log = await _logCollection.AddAsync(JobId.Empty, leaseId, agent.SessionId, LogType.Json, cancellationToken: cancellationToken);
 							task.LogId = log.Id.ToString();
 							task.RemoveUntrackedFiles = agent.RequestFullConform;
-							task.SkipWhenDiskFreeSpaceIsAvailable = agent.RequestConformSkipWhenDiskFreeSpaceIsAvailable;
 
 							return LeaseAsync(new CreateLeaseOptions(leaseId, null, "Updating workspaces", null, null, log.Id, null, true, task));
 						}
diff --git a/Engine/Source/Programs/Horde/Plugins/Build/HordeServer.Build/Utilities/ConformUtilities.cs b/Engine/Source/Programs/Horde/Plugins/Build/HordeServer.Build/Utilities/ConformUtilities.cs
--- a/Engine/Source/Programs/Horde/Plugins/Build/HordeServer.Build/Utilities/ConformUtilities.cs
+++ b/Engine/Source/Programs/Horde/Plugins/Build/HordeServer.Build/Utilities/ConformUtilities.cs
@@ -19,10 +19,9 @@ public static class ConformUtilities
 	/// <param name="poolService">Pool service instance</param>
 	/// <param name="config">Current configuration</param>
 	/// <param name="logger">Logger instance</param>
-	/// <param name="skipWhenDiskFreeSpaceIsAvailable">Whether upon execution of the conform should the conform be skipped if the agent now has the available disk free space required</param>
 	/// <param name="cancellationToken">Cancellation token for the async task</param>
 	/// <returns>True if conform was requested for the agent; otherwise false</returns>
-	public static async Task<bool> AutoConformAgentAsync(IAgent agent, PoolService poolService, BuildConfig config, bool skipWhenDiskFreeSpaceIsAvailable, ILogger logger, CancellationToken cancellationToken)
+	public static async Task<bool> AutoConformAgentAsync(IAgent agent, PoolService poolService, BuildConfig config, ILogger logger, CancellationToken cancellationToken)
 	{
 		long MegabytesToBytes(long v) => v * 1024 * 1024;
 		double BytesToMegabytes(long v) => v / 1024.0 / 1024.0;
@@ -44,12 +43,7 @@ public static class ConformUtilities
 		if (freeDiskSpace != null && conformDiskSpaceNeeded > 0 && freeDiskSpace < conformDiskSpaceNeeded &&
 		    agent.ConformAttemptCount is null or 0 && !agent.RequestFullConform)
 		{
-			UpdateAgentOptions options = new() { RequestFullConform = true };
-			if (skipWhenDiskFreeSpaceIsAvailable)
-			{
-				options = options with { RequestConformSkipWhenDiskFreeSpaceIsAvailable = true };
-			}
-			newAgent = await agent.TryUpdateAsync(options, cancellationToken);
+			newAgent = await agent.TryUpdateAsync(new UpdateAgentOptions { RequestFullConform = true }, cancellationToken);
 			logger.LogInformation("Auto-conforming {AgentId} as workspace conform disk space needed ({ConformDiskSpace:F1} MB) is greater than free disk space ({FreeDiskSpace:F1} MB)",
 				agent.Id.ToString(), BytesToMegabytes(conformDiskSpaceNeeded), BytesToMegabytes(freeDiskSpace.Value));
 		}
diff --git a/Engine/Source/Programs/Horde/Plugins/Compute/HordeServer.Compute.Tests/Agents/AgentCollectionTests.cs b/Engine/Source/Programs/Horde/Plugins/Compute/HordeServer.Compute.Tests/Agents/AgentCollectionTests.cs
--- a/Engine/Source/Programs/Horde/Plugins/Compute/HordeServer.Compute.Tests/Agents/AgentCollectionTests.cs
+++ b/Engine/Source/Programs/Horde/Plugins/Compute/HordeServer.Compute.Tests/Agents/AgentCollectionTests.cs
@@ -379,6 +379,145 @@ public class AgentCollectionTests : ComputeTestSetup
 		Assert.IsTrue(lease.OutcomeReason.HasFlag(LeaseOutcomeReason.NonZeroExitCode));
 	}
 
+	[TestMethod]
+	public async Task UpdateSession_CapabilityOverridesReplacesExistingPropertyAsync()
+	{
+		// Boot with initial WorkingDiskFreeSpace in capabilities
+		List<string> initialProps = [$"{KnownPropertyNames.WorkingDiskFreeSpace}=500"];
+		IAgent? agent = Deref(await _agent.TryCreateSessionAsync(new CreateSessionOptions(new RpcAgentCapabilities(initialProps), [], "foo")));
+		Assert.AreEqual(500, agent.GetDiskFreeSpace());
+
+		// Update replaces (not duplicates) the property via CapabilityOverrides
+		RpcAgentCapabilities overrides = new([$"{KnownPropertyNames.WorkingDiskFreeSpace}=250"]);
+		agent = Deref(await agent.TryUpdateSessionAsync(new UpdateSessionOptions { CapabilityOverrides = overrides }));
+		Assert.AreEqual(250, agent.GetDiskFreeSpace());
+		Assert.HasCount(1, agent.GetPropertyValues(KnownPropertyNames.WorkingDiskFreeSpace));
+
+		// Re-read from DB to verify replacement persisted
+		IAgent? agentFromDb = Deref(await AgentCollection.GetAsync(agent.Id));
+		Assert.AreEqual(250, agentFromDb.GetDiskFreeSpace());
+		Assert.HasCount(1, agentFromDb.GetPropertyValues(KnownPropertyNames.WorkingDiskFreeSpace));
+
+		// Update without CapabilityOverrides preserves existing value
+		agent = Deref(await agent.TryUpdateSessionAsync(new UpdateSessionOptions { Status = AgentStatus.Ok }));
+		Assert.AreEqual(250, agent.GetDiskFreeSpace());
+
+		agentFromDb = Deref(await AgentCollection.GetAsync(agent.Id));
+		Assert.AreEqual(250, agentFromDb.GetDiskFreeSpace());
+	}
+
+	[TestMethod]
+	public async Task UpdateSession_CapabilityOverridesReplacesExistingResourceAsync()
+	{
+		// Boot with initial resources
+		RpcAgentCapabilities initialCaps = new();
+		initialCaps.Resources.Add("LogicalCores", 16);
+		initialCaps.Resources.Add("Ram", 32768);
+		IAgent? agent = Deref(await _agent.TryCreateSessionAsync(new CreateSessionOptions(initialCaps, [], "foo")));
+		Assert.AreEqual(16, agent.Resources["LogicalCores"]);
+		Assert.AreEqual(32768, agent.Resources["Ram"]);
+
+		// Override one resource, leave the other unchanged
+		RpcAgentCapabilities overrides = new();
+		overrides.Resources.Add("Ram", 65536);
+		agent = Deref(await agent.TryUpdateSessionAsync(new UpdateSessionOptions { CapabilityOverrides = overrides }));
+		Assert.AreEqual(16, agent.Resources["LogicalCores"]);
+		Assert.AreEqual(65536, agent.Resources["Ram"]);
+
+		// Re-read from DB to verify
+		IAgent? agentFromDb = Deref(await AgentCollection.GetAsync(agent.Id));
+		Assert.AreEqual(16, agentFromDb.Resources["LogicalCores"]);
+		Assert.AreEqual(65536, agentFromDb.Resources["Ram"]);
+	}
+
+	[TestMethod]
+	public async Task UpdateSession_CapabilityOverridesReplacesResourceAsync()
+	{
+		IAgent? agent = Deref(await _agent.TryCreateSessionAsync(new CreateSessionOptions(new RpcAgentCapabilities(), [], "foo")));
+
+		RpcAgentCapabilities caps = new();
+		caps.Resources.Add("LogicalCores", 16);
+		caps.Resources.Add("Ram", 32768);
+
+		RpcAgentCapabilities overrides = new();
+		overrides.Resources.Add("Ram", 65536);
+
+		agent = Deref(await agent.TryUpdateSessionAsync(new UpdateSessionOptions { Capabilities = caps, CapabilityOverrides = overrides }));
+		Assert.AreEqual(16, agent.Resources["LogicalCores"]);
+		Assert.AreEqual(65536, agent.Resources["Ram"]);
+
+		// Re-read from DB to verify
+		IAgent? agentFromDb = Deref(await AgentCollection.GetAsync(agent.Id));
+		Assert.AreEqual(16, agentFromDb.Resources["LogicalCores"]);
+		Assert.AreEqual(65536, agentFromDb.Resources["Ram"]);
+	}
+
+	[TestMethod]
+	public async Task UpdateSession_CapabilityOverridesReplacesMultiPropertiesAsync()
+	{
+		// Boot with multiple properties
+		List<string> initialProps = ["OSFamily=Windows", "OsVersion=10.0", $"{KnownPropertyNames.WorkingDiskFreeSpace}=500"];
+		IAgent? agent = Deref(await _agent.TryCreateSessionAsync(new CreateSessionOptions(new RpcAgentCapabilities(initialProps), [], "foo")));
+		Assert.AreEqual(500, agent.GetDiskFreeSpace());
+		Assert.HasCount(1, agent.GetPropertyValues("OSFamily"));
+
+		// Override WorkingDiskFreeSpace and OSFamily, leave OsVersion unchanged
+		RpcAgentCapabilities overrides = new([$"{KnownPropertyNames.WorkingDiskFreeSpace}=250", "OSFamily=Linux"]);
+		agent = Deref(await agent.TryUpdateSessionAsync(new UpdateSessionOptions { CapabilityOverrides = overrides }));
+		Assert.AreEqual(250, agent.GetDiskFreeSpace());
+		Assert.AreEqual("Linux", agent.GetPropertyValues("OSFamily").First());
+		Assert.AreEqual("10.0", agent.GetPropertyValues("OsVersion").First());
+
+		// Re-read from DB to verify
+		IAgent? agentFromDb = Deref(await AgentCollection.GetAsync(agent.Id));
+		Assert.AreEqual(250, agentFromDb.GetDiskFreeSpace());
+		Assert.AreEqual("Linux", agentFromDb.GetPropertyValues("OSFamily").First());
+		Assert.AreEqual("10.0", agentFromDb.GetPropertyValues("OsVersion").First());
+	}
+
+	[TestMethod]
+	public async Task UpdateSession_CapabilityOverridesIgnoresNewPropertyAsync()
+	{
+		// Boot with a single property
+		List<string> initialProps = ["OSFamily=Windows"];
+		IAgent? agent = Deref(await _agent.TryCreateSessionAsync(new CreateSessionOptions(new RpcAgentCapabilities(initialProps), [], "foo")));
+		Assert.AreEqual("Windows", agent.GetPropertyValues("OSFamily").First());
+
+		// Override with one existing key and one new key; new key should be ignored
+		RpcAgentCapabilities overrides = new(["OSFamily=Linux", "NewProp=ShouldBeIgnored"]);
+		agent = Deref(await agent.TryUpdateSessionAsync(new UpdateSessionOptions { CapabilityOverrides = overrides }));
+		Assert.AreEqual("Linux", agent.GetPropertyValues("OSFamily").First());
+		Assert.HasCount(0, agent.GetPropertyValues("NewProp"));
+
+		// Re-read from DB to verify
+		IAgent? agentFromDb = Deref(await AgentCollection.GetAsync(agent.Id));
+		Assert.AreEqual("Linux", agentFromDb.GetPropertyValues("OSFamily").First());
+		Assert.HasCount(0, agentFromDb.GetPropertyValues("NewProp"));
+	}
+
+	[TestMethod]
+	public async Task UpdateSession_CapabilityOverridesIgnoresNewResourceAsync()
+	{
+		// Boot with initial resources
+		RpcAgentCapabilities initialCaps = new();
+		initialCaps.Resources.Add("LogicalCores", 16);
+		IAgent? agent = Deref(await _agent.TryCreateSessionAsync(new CreateSessionOptions(initialCaps, [], "foo")));
+		Assert.AreEqual(16, agent.Resources["LogicalCores"]);
+
+		// Override with one existing key and one new key; new key should be ignored
+		RpcAgentCapabilities overrides = new();
+		overrides.Resources.Add("LogicalCores", 32);
+		overrides.Resources.Add("Gpu", 2);
+		agent = Deref(await agent.TryUpdateSessionAsync(new UpdateSessionOptions { CapabilityOverrides = overrides }));
+		Assert.AreEqual(32, agent.Resources["LogicalCores"]);
+		Assert.IsFalse(agent.Resources.ContainsKey("Gpu"));
+
+		// Re-read from DB to verify
+		IAgent? agentFromDb = Deref(await AgentCollection.GetAsync(agent.Id));
+		Assert.AreEqual(32, agentFromDb.Resources["LogicalCores"]);
+		Assert.IsFalse(agentFromDb.Resources.ContainsKey("Gpu"));
+	}
+
 	[TestMethod]
 	public async Task GetChildLeaseIdsAsync()
 	{
diff --git a/Engine/Source/Programs/Horde/Plugins/Compute/HordeServer.Compute/Agents/AgentCollection.cs b/Engine/Source/Programs/Horde/Plugins/Compute/HordeServer.Compute/Agents/AgentCollection.cs
--- a/Engine/Source/Programs/Horde/Plugins/Compute/HordeServer.Compute/Agents/AgentCollection.cs
+++ b/Engine/Source/Programs/Horde/Plugins/Compute/HordeServer.Compute/Agents/AgentCollection.cs
@@ -69,7 +69,6 @@ namespace HordeServer.Agents
 			IReadOnlyList<PoolId> IAgent.ExplicitPools => _document.ExplicitPools;
 			bool IAgent.RequestConform => _document.RequestConform;
 			bool IAgent.RequestFullConform => _document.RequestFullConform;
-			bool IAgent.RequestConformSkipWhenDiskFreeSpaceIsAvailable => _document.RequestConformSkipWhenDiskFreeSpaceIsAvailable;
 			bool IAgent.RequestRestart => _document.RequestRestart;
 			bool IAgent.RequestShutdown => _document.RequestShutdown;
 			bool IAgent.RequestForceRestart => _document.RequestForceRestart;
@@ -190,9 +189,6 @@ namespace HordeServer.Agents
 			[BsonIgnoreIfDefault, BsonDefaultValue(false)]
 			public bool RequestFullConform { get; set; }
 
-			[BsonIgnoreIfDefault, BsonDefaultValue(false)]
-			public bool RequestConformSkipWhenDiskFreeSpaceIsAvailable { get; set; }
-
 			[BsonIgnoreIfDefault, BsonDefaultValue(false)]
 			public bool RequestRestart { get; set; }
 
@@ -798,23 +794,11 @@ namespace HordeServer.Agents
 			{
 				updates.Add(updateBuilder.Set(x => x.RequestConform, options.RequestConform.Value));
 				updates.Add(updateBuilder.Unset(x => x.ConformAttemptCount));
-				if (options.RequestFullConform == null)
-				{
-					updates.Add(updateBuilder.Unset(x => x.RequestConformSkipWhenDiskFreeSpaceIsAvailable));
-				}
 			}
 			if (options.RequestFullConform != null)
 			{
 				updates.Add(updateBuilder.Set(x => x.RequestFullConform, options.RequestFullConform.Value));
 				updates.Add(updateBuilder.Unset(x => x.ConformAttemptCount));
-				if (options.RequestConformSkipWhenDiskFreeSpaceIsAvailable is true)
-				{
-					updates.Add(updateBuilder.Set(x => x.RequestConformSkipWhenDiskFreeSpaceIsAvailable, true));
-				}
-				else
-				{
-					updates.Add(updateBuilder.Unset(x => x.RequestConformSkipWhenDiskFreeSpaceIsAvailable));
-				}
 			}
 			if (options.RequestRestart != null)
 			{
@@ -878,16 +862,26 @@ namespace HordeServer.Agents
 				UpdateDefinitionBuilder<AgentDocument> updateBuilder = Builders<AgentDocument>.Update;
 				List<UpdateDefinition<AgentDocument>> updates = new List<UpdateDefinition<AgentDocument>>();
 
-				List<string> newProperties = newDocument.Properties ?? new List<string>();
+				List<string> newProperties = newDocument.Properties ?? [];
+				Dictionary<string, int>? newResources = null;
 				if (options.Capabilities != null)
 				{
-					newProperties = new List<string>(options.Capabilities.Properties);
+					newProperties = [.. options.Capabilities.Properties];
+					newResources = options.Capabilities.Resources.ToDictionary();
+				}
 
-					if (!ResourcesEqual(options.Capabilities.Resources, newDocument.Resources))
+				if (options.CapabilityOverrides != null)
+				{
+					if (options.CapabilityOverrides.Properties.Count > 0)
 					{
-						updates.Add(updateBuilder.Set(x => x.Resources, options.Capabilities.Resources.ToDictionary()));
+						newProperties = MergeOverrideProperties(agentId, newProperties, options.CapabilityOverrides.Properties);
+					}
+					if (options.CapabilityOverrides.Resources.Count > 0)
+					{
+						newResources = MergeOverrideResources(agentId, newResources ?? newDocument.Resources, options.CapabilityOverrides.Resources);
 					}
 				}
+
 				if (options.DynamicPools != null)
 				{
 					List<PoolId> dynamicPools = CreatePoolsList(options.DynamicPools).ToList();
@@ -905,6 +899,11 @@ namespace HordeServer.Agents
 
 				SetPoolProperties(newProperties, pools);
 
+				if (newResources != null && !ResourcesEqual(newResources, newDocument.Resources))
+				{
+					updates.Add(updateBuilder.Set(x => x.Resources, newResources));
+				}
+
 				if (!Enumerable.SequenceEqual(newDocument.Properties ?? Enumerable.Empty<string>(), newProperties, StringComparer.Ordinal))
 				{
 					updates.Add(updateBuilder.Set(x => x.Properties, newProperties));
@@ -1070,6 +1069,76 @@ namespace HordeServer.Agents
 			properties.Sort(StringComparer.OrdinalIgnoreCase);
 		}
 
+		/// <summary>
+		/// Merges override properties into a property list. For each key present in both the overrides
+		/// and the existing properties, existing entries with that key are replaced by the override entries.
+		/// Override entries with keys not present in the existing properties are ignored.
+		/// </summary>
+		List<string> MergeOverrideProperties(AgentId agentId, IEnumerable<string> properties, IReadOnlyList<string> overrides)
+		{
+			List<string> newProperties = [.. properties];
+
+			HashSet<string> existingKeys = new(StringComparer.OrdinalIgnoreCase);
+			foreach (string prop in newProperties)
+			{
+				existingKeys.Add(GetPropertyKey(prop));
+			}
+
+			HashSet<string> overrideKeys = new(StringComparer.OrdinalIgnoreCase);
+			foreach (string prop in overrides)
+			{
+				string key = GetPropertyKey(prop);
+				if (existingKeys.Contains(key))
+				{
+					overrideKeys.Add(key);
+				}
+				else
+				{
+					_logger.LogWarning("Agent {AgentId}: capability override property {Key} ignored - key not present in existing properties", agentId, key);
+				}
+			}
+
+			newProperties.RemoveAll(x => overrideKeys.Contains(GetPropertyKey(x)));
+			newProperties.AddRange(overrides.Where(x => overrideKeys.Contains(GetPropertyKey(x))));
+			newProperties.Sort(StringComparer.OrdinalIgnoreCase);
+			return newProperties;
+		}
+
+		/// <summary>
+		/// Merges override resources into an existing resource dictionary, replacing values for existing keys.
+		/// Override entries with keys not present in the existing resources are ignored.
+		/// </summary>
+		Dictionary<string, int>? MergeOverrideResources(AgentId agentId, IReadOnlyDictionary<string, int>? resources, IReadOnlyDictionary<string, int> overrides)
+		{
+			if (resources == null)
+			{
+				return null;
+			}
+			Dictionary<string, int> newResources = new(resources, StringComparer.OrdinalIgnoreCase);
+			foreach ((string key, int value) in overrides)
+			{
+				if (newResources.ContainsKey(key))
+				{
+					newResources[key] = value;
+				}
+				else
+				{
+					_logger.LogWarning("Agent {AgentId}: capability override resource {Key} ignored - key not present in existing resources", agentId, key);
+				}
+			}
+			return newResources;
+		}
+
+		/// <summary>
+		/// Extracts the key portion from a "KEY=VALUE" property string.
+		/// Returns the full string if no '=' is found.
+		/// </summary>
+		static string GetPropertyKey(string property)
+		{
+			int eqIdx = property.IndexOf('=', StringComparison.Ordinal);
+			return eqIdx >= 0 ? property[..eqIdx] : property;
+		}
+
 		static List<PoolId> CreatePoolsList(IEnumerable<PoolId> pools)
 			=> pools.Distinct().OrderBy(x => x.Id.Text).ToList();
 
@@ -1132,7 +1201,6 @@ namespace HordeServer.Agents
 			UpdateDefinition<AgentDocument> update = Builders<AgentDocument>.Update.Set(x => x.Workspaces, workspaces);
 			update = update.Set(x => x.LastConformTime, lastConformTime);
 			update = update.Unset(x => x.ConformAttemptCount);
-			update = update.Unset(x => x.RequestConformSkipWhenDiskFreeSpaceIsAvailable);
 			if (!requestConform)
 			{
 				update = update.Unset(x => x.RequestConform);
diff --git a/Engine/Source/Programs/Horde/Plugins/Compute/HordeServer.Compute/Agents/AgentService.cs b/Engine/Source/Programs/Horde/Plugins/Compute/HordeServer.Compute/Agents/AgentService.cs
--- a/Engine/Source/Programs/Horde/Plugins/Compute/HordeServer.Compute/Agents/AgentService.cs
+++ b/Engine/Source/Programs/Horde/Plugins/Compute/HordeServer.Compute/Agents/AgentService.cs
@@ -615,9 +615,9 @@ namespace HordeServer.Agents
 		}
 
 		/// <summary>
-		/// 
+		///
 		/// </summary>
-		public async Task<IAgent?> UpdateSessionWithWaitAsync(IAgent inAgent, SessionId sessionId, AgentStatus status, RpcAgentCapabilities? capabilities, IList<HordeCommon.Rpc.Messages.RpcLease> newLeases, CancellationToken cancellationToken = default)
+		public async Task<IAgent?> UpdateSessionWithWaitAsync(IAgent inAgent, SessionId sessionId, AgentStatus status, RpcAgentCapabilities? capabilities, IList<HordeCommon.Rpc.Messages.RpcLease> newLeases, RpcAgentCapabilities? capabilityOverrides = null, CancellationToken cancellationToken = default)
 		{
 			IAgent? agent = inAgent;
 
@@ -625,7 +625,7 @@ namespace HordeServer.Agents
 			uint updateIndex = agent.UpdateIndex;
 
 			// Update the agent session and return to the caller if anything changes
-			agent = await UpdateSessionAsync(agent, sessionId, status, capabilities, newLeases, cancellationToken);
+			agent = await UpdateSessionAsync(agent, sessionId, status, capabilities, newLeases, capabilityOverrides, cancellationToken);
 			if (agent != null && agent.UpdateIndex == updateIndex && (agent.Leases.Count > 0 || agent.Status != AgentStatus.Stopping))
 			{
 				_logger.LogDebug("Waiting for lease update on agent {AgentId} session {SessionId}", agent.Id, sessionId);
@@ -642,9 +642,10 @@ namespace HordeServer.Agents
 		/// <param name="status">New status for the agent</param>
 		/// <param name="capabilities">New agent capabilities</param>
 		/// <param name="newLeases">New list of leases for this session</param>
+		/// <param name="capabilityOverrides">Capability overrides to merge into existing agent capabilities</param>
 		/// <param name="cancellationToken">Cancellation token for the operation</param>
 		/// <returns>Updated agent state</returns>
-		public async Task<IAgent?> UpdateSessionAsync(IAgent inAgent, SessionId sessionId, AgentStatus status, RpcAgentCapabilities? capabilities, IList<HordeCommon.Rpc.Messages.RpcLease> newLeases, CancellationToken cancellationToken = default)
+		public async Task<IAgent?> UpdateSessionAsync(IAgent inAgent, SessionId sessionId, AgentStatus status, RpcAgentCapabilities? capabilities, IList<HordeCommon.Rpc.Messages.RpcLease> newLeases, RpcAgentCapabilities? capabilityOverrides = null, CancellationToken cancellationToken = default)
 		{
 			DateTime utcNow = _clock.UtcNow;
 
@@ -674,7 +675,7 @@ namespace HordeServer.Agents
 				List<PoolId> dynamicPools = await GetDynamicPoolsAsync(agent, cancellationToken);
 
 				// Update the agent, and try to create new lease documents if we succeed
-				IAgent? newAgent = await agent.TryUpdateSessionAsync(new UpdateSessionOptions(status, capabilities, dynamicPools, newLeases), cancellationToken);
+				IAgent? newAgent = await agent.TryUpdateSessionAsync(new UpdateSessionOptions(status, capabilities, dynamicPools, newLeases, CapabilityOverrides: capabilityOverrides), cancellationToken);
 				if (newAgent != null)
 				{
 					agent = newAgent;
diff --git a/Engine/Source/Programs/Horde/Plugins/Compute/HordeServer.Compute/Agents/AgentsController.cs b/Engine/Source/Programs/Horde/Plugins/Compute/HordeServer.Compute/Agents/AgentsController.cs
--- a/Engine/Source/Programs/Horde/Plugins/Compute/HordeServer.Compute/Agents/AgentsController.cs
+++ b/Engine/Source/Programs/Horde/Plugins/Compute/HordeServer.Compute/Agents/AgentsController.cs
@@ -235,7 +235,6 @@ namespace HordeServer.Agents
 				agent.Deleted,
 				agent.RequestConform,
 				agent.RequestFullConform,
-				agent.RequestConformSkipWhenDiskFreeSpaceIsAvailable,
 				agent.RequestRestart,
 				agent.RequestShutdown,
 				agent.LastShutdownReason ?? "Unknown",
@@ -288,7 +287,7 @@ namespace HordeServer.Agents
 					return NotFound(agentId);
 				}
 
-				IAgent? newAgent = await agent.TryUpdateAsync(new UpdateAgentOptions(update.Enabled, update.RequestConform, update.RequestFullConform, update.RequestConformSkipWhenAvailableFreeDiskSpace, update.RequestRestart, update.RequestShutdown, update.RequestForceRestart, $"Manual ({userName})", update.Pools?.ConvertAll(x => new PoolId(x)), update.Comment));
+				IAgent? newAgent = await agent.TryUpdateAsync(new UpdateAgentOptions(update.Enabled, update.RequestConform, update.RequestFullConform, update.RequestRestart, update.RequestShutdown, update.RequestForceRestart, $"Manual ({userName})", update.Pools?.ConvertAll(x => new PoolId(x)), update.Comment));
 				if (newAgent == null)
 				{
 					continue;
@@ -307,10 +306,6 @@ namespace HordeServer.Agents
 				{
 					logger.LogInformation("Setting changed: RequestFullConform = {State} ({UserName} via IP {Ip})", newAgent.RequestFullConform, userName, clientIp);
 				}
-				if (agent.RequestConformSkipWhenDiskFreeSpaceIsAvailable != newAgent.RequestConformSkipWhenDiskFreeSpaceIsAvailable)
-				{
-					logger.LogInformation("Setting changed: RequestConformSkipWhenDiskFreeSpaceIsAvailable = {Value} ({UserName} via IP {Ip})", newAgent.RequestConformSkipWhenDiskFreeSpaceIsAvailable, userName, clientIp);
-				}
 				if (agent.RequestRestart != newAgent.RequestRestart)
 				{
 					logger.LogInformation("Setting changed: RequestRestart = {State} ({UserName} via IP {Ip})", newAgent.RequestRestart, userName, clientIp);
diff --git a/Engine/Source/Programs/Horde/Plugins/Compute/HordeServer.Compute/Agents/IAgent.cs b/Engine/Source/Programs/Horde/Plugins/Compute/HordeServer.Compute/Agents/IAgent.cs
--- a/Engine/Source/Programs/Horde/Plugins/Compute/HordeServer.Compute/Agents/IAgent.cs
+++ b/Engine/Source/Programs/Horde/Plugins/Compute/HordeServer.Compute/Agents/IAgent.cs
@@ -167,11 +167,6 @@ namespace HordeServer.Agents
 		/// </summary>
 		public bool RequestFullConform { get; }
 
-		/// <summary>
-		/// Whether a conform can be skipped when the machine already has the required available disk free space needed
-		/// </summary>
-		public bool RequestConformSkipWhenDiskFreeSpaceIsAvailable { get; }
-
 		/// <summary>
 		/// Whether a machine restart is requested
 		/// </summary>
@@ -346,14 +341,13 @@ namespace HordeServer.Agents
 	/// <param name="Enabled">Whether the agent is enabled or not</param>
 	/// <param name="RequestConform">Whether to request a conform job be run</param>
 	/// <param name="RequestFullConform">Whether to request a full conform job be run</param>
-	/// <param name="RequestConformSkipWhenDiskFreeSpaceIsAvailable">Whether a conform can be skipped when the agent already has the available disk free space needed</param>
 	/// <param name="RequestRestart">Whether to request the machine be restarted</param>
 	/// <param name="RequestShutdown">Whether to request the machine be shut down</param>
 	/// <param name="RequestForceRestart">Request an immediate restart without waiting for leases to complete</param>
 	/// <param name="ShutdownReason">The reason for shutting down agent, ex. Autoscaler/Manual/Unexpected</param>
 	/// <param name="ExplicitPools">List of pools for the agent</param>
 	/// <param name="Comment">New comment</param>
-	public record UpdateAgentOptions(bool? Enabled = null, bool? RequestConform = null, bool? RequestFullConform = null, bool? RequestConformSkipWhenDiskFreeSpaceIsAvailable = null, bool? RequestRestart = null, bool? RequestShutdown = null, bool? RequestForceRestart = null, string? ShutdownReason = null, List<PoolId>? ExplicitPools = null, string? Comment = null);
+	public record UpdateAgentOptions(bool? Enabled = null, bool? RequestConform = null, bool? RequestFullConform = null, bool? RequestRestart = null, bool? RequestShutdown = null, bool? RequestForceRestart = null, string? ShutdownReason = null, List<PoolId>? ExplicitPools = null, string? Comment = null);
 
 	/// <summary>
 	/// Options for starting a new agent session
@@ -371,7 +365,8 @@ namespace HordeServer.Agents
 	/// <param name="DynamicPools">New list of dynamic pools for the agent</param>
 	/// <param name="Leases">New set of leases</param>
 	/// <param name="DefaultLeaseOutcomeReason">Default reason used when terminating leases without agent-provided reasons</param>
-	public record UpdateSessionOptions(AgentStatus? Status = null, RpcAgentCapabilities? Capabilities = null, IReadOnlyList<PoolId>? DynamicPools = null, IEnumerable<RpcLease>? Leases = null, LeaseOutcomeReason? DefaultLeaseOutcomeReason = null);
+	/// <param name="CapabilityOverrides">Capability overrides to merge into existing agent capabilities</param>
+	public record UpdateSessionOptions(AgentStatus? Status = null, RpcAgentCapabilities? Capabilities = null, IReadOnlyList<PoolId>? DynamicPools = null, IEnumerable<RpcLease>? Leases = null, LeaseOutcomeReason? DefaultLeaseOutcomeReason = null, RpcAgentCapabilities? CapabilityOverrides = null);
 
 	/// <summary>
 	/// Extension methods for IAgent
diff --git a/Engine/Source/Programs/Horde/Plugins/Compute/HordeServer.Compute/Server/RpcService.cs b/Engine/Source/Programs/Horde/Plugins/Compute/HordeServer.Compute/Server/RpcService.cs
--- a/Engine/Source/Programs/Horde/Plugins/Compute/HordeServer.Compute/Server/RpcService.cs
+++ b/Engine/Source/Programs/Horde/Plugins/Compute/HordeServer.Compute/Server/RpcService.cs
@@ -319,7 +319,8 @@ namespace HordeServer.Server
 						// Update the session
 						try
 						{
-							agent = await _agentService.UpdateSessionWithWaitAsync(agent, sessionId, (AgentStatus)request.Status, request.Capabilities?.MergeDevices(), request.Leases, cancellationSource.Token);
+							RpcAgentCapabilities? capabilityOverrides = request.CapabilityOverrides?.Properties.Count == 0 && request.CapabilityOverrides?.Resources.Count == 0 ? null : request.CapabilityOverrides;
+							agent = await _agentService.UpdateSessionWithWaitAsync(agent, sessionId, (AgentStatus)request.Status, request.Capabilities?.MergeDevices(), request.Leases, capabilityOverrides, cancellationSource.Token);
 						}
 						catch (OperationCanceledException)
 						{
diff --git a/Engine/Source/Programs/Shared/EpicGames.Horde/Agents/AgentMessages.cs b/Engine/Source/Programs/Shared/EpicGames.Horde/Agents/AgentMessages.cs
--- a/Engine/Source/Programs/Shared/EpicGames.Horde/Agents/AgentMessages.cs
+++ b/Engine/Source/Programs/Shared/EpicGames.Horde/Agents/AgentMessages.cs
@@ -51,13 +51,12 @@ namespace EpicGames.Horde.Agents
 	/// <param name="Ephemeral"> Whether the agent is ephemeral </param>
 	/// <param name="RequestConform"> Request a conform be performed using the current agent set </param>
 	/// <param name="RequestFullConform"> Request that a full conform be performed, removing all intermediate files </param>
-	/// <param name="RequestConformSkipWhenAvailableFreeDiskSpace"> Request that a conform can be skipped when the machine already has the available free disk space needed </param>
 	/// <param name="RequestRestart"> Request the machine be restarted </param>
 	/// <param name="RequestShutdown"> Request the machine be shut down </param>
 	/// <param name="RequestForceRestart"> Request the machine be restarted without waiting for leases to complete </param>
 	/// <param name="Pools"> Pools for this agent </param>
 	/// <param name="Comment"> New comment </param>
-	public record UpdateAgentRequest(bool? Enabled = null, bool? Ephemeral = null, bool? RequestConform = null, bool? RequestFullConform = null, bool? RequestConformSkipWhenAvailableFreeDiskSpace = null, bool? RequestRestart = null, bool? RequestShutdown = null, bool? RequestForceRestart = null, List<string>? Pools = null, string? Comment = null);
+	public record UpdateAgentRequest(bool? Enabled = null, bool? Ephemeral = null, bool? RequestConform = null, bool? RequestFullConform = null, bool? RequestRestart = null, bool? RequestShutdown = null, bool? RequestForceRestart = null, List<string>? Pools = null, string? Comment = null);
 
 	/// <summary>
 	/// Response for queries to find a particular lease within an agent
@@ -106,7 +105,6 @@ namespace EpicGames.Horde.Agents
 	/// <param name="Deleted">Whether this agent has expired</param>
 	/// <param name="PendingConform">Whether a conform job is pending</param>
 	/// <param name="PendingFullConform">Whether a full conform job is pending</param>
-	/// <param name="PendingConformSkipWhenDiskFreeSpaceIsAvailable">Whether a conform can be skipped when the agent already has the available disk free space needed</param>
 	/// <param name="PendingRestart">Whether a restart is pending</param>
 	/// <param name="PendingShutdown">Whether a shutdown is pending</param>
 	/// <param name="LastShutdownReason">The reason for the last shutdown</param>
@@ -123,7 +121,7 @@ namespace EpicGames.Horde.Agents
 	/// <param name="Leases">Array of active leases.</param>
 	/// <param name="Workspaces">Current workspaces synced on the agent</param>
 	/// <param name="Comment">Comment for this agent</param>
-	public record GetAgentResponse(AgentId Id, string Name, bool Enabled, AgentStatus Status, double? Rate, SessionId? SessionId, bool Ephemeral, bool Online, bool Deleted, bool PendingConform, bool PendingFullConform, bool PendingConformSkipWhenDiskFreeSpaceIsAvailable, bool PendingRestart, bool PendingShutdown, string LastShutdownReason, DateTime LastConformTime, int? ConformAttemptCount, DateTime? NextConformTime, string? Version, List<string> Properties, Dictionary<string, int> Resources, DateTime? UpdateTime, DateTime? LastOnlineTime, List<string>? Pools, object? Capabilities, List<GetAgentLeaseResponse> Leases, List<GetAgentWorkspaceResponse> Workspaces, string? Comment);
+	public record GetAgentResponse(AgentId Id, string Name, bool Enabled, AgentStatus Status, double? Rate, SessionId? SessionId, bool Ephemeral, bool Online, bool Deleted, bool PendingConform, bool PendingFullConform, bool PendingRestart, bool PendingShutdown, string LastShutdownReason, DateTime LastConformTime, int? ConformAttemptCount, DateTime? NextConformTime, string? Version, List<string> Properties, Dictionary<string, int> Resources, DateTime? UpdateTime, DateTime? LastOnlineTime, List<string>? Pools, object? Capabilities, List<GetAgentLeaseResponse> Leases, List<GetAgentWorkspaceResponse> Workspaces, string? Comment);
 
 	/// <summary>
 	/// Telemetry data for an agent
