Horde: Add session conflict detection and auto-quarantine to agents

Detect agents experiencing repeated session ID mismatches (e.g. cloned VMs sharing an agent ID) and automatically quarantine them after a configurable threshold. Conflicts are tracked in Redis using a sliding-window sorted set; once the threshold is reached the agent is disabled with a diagnostic comment listing the offending IPs. Suppresses log spam for already-quarantined agents and surfaces them in the daily Slack agent report.

diff --git a/Engine/Source/Programs/Horde/HordeServer.Tests.Shared/HttpContextStub.cs b/Engine/Source/Programs/Horde/HordeServer.Tests.Shared/HttpContextStub.cs
--- a/Engine/Source/Programs/Horde/HordeServer.Tests.Shared/HttpContextStub.cs
+++ b/Engine/Source/Programs/Horde/HordeServer.Tests.Shared/HttpContextStub.cs
@@ -1,6 +1,8 @@
 // Copyright Epic Games, Inc. All Rights Reserved.
 
+using System.Net;
 using System.Security.Claims;
+using System.Security.Cryptography.X509Certificates;
 using Microsoft.AspNetCore.Http;
 using Microsoft.AspNetCore.Http.Features;
 
@@ -8,9 +10,39 @@ namespace HordeServer.Tests
 {
 	using ISession = Microsoft.AspNetCore.Http.ISession;
 
+	/// <summary>
+	/// Stub implementation of <see cref="ConnectionInfo"/> for testing
+	/// </summary>
+	public sealed class ConnectionInfoStub : ConnectionInfo
+	{
+		/// <inheritdoc/>
+		public override string Id { get; set; } = "test-connection";
+
+		/// <inheritdoc/>
+		public override IPAddress? RemoteIpAddress { get; set; }
+
+		/// <inheritdoc/>
+		public override int RemotePort { get; set; }
+
+		/// <inheritdoc/>
+		public override IPAddress? LocalIpAddress { get; set; }
+
+		/// <inheritdoc/>
+		public override int LocalPort { get; set; }
+
+		/// <inheritdoc/>
+		public override X509Certificate2? ClientCertificate { get; set; }
+
+		/// <inheritdoc/>
+		public override Task<X509Certificate2?> GetClientCertificateAsync(CancellationToken cancellationToken = default)
+		{
+			return Task.FromResult(ClientCertificate);
+		}
+	}
+
 	sealed class HttpContextStub : HttpContext
 	{
-		public override ConnectionInfo Connection { get; } = null!;
+		public override ConnectionInfo Connection { get; }
 		public override IFeatureCollection Features { get; } = null!;
 		public override IDictionary<object, object?> Items { get; set; } = null!;
 		public override HttpRequest Request { get; } = null!;
@@ -22,16 +54,15 @@ namespace HordeServer.Tests
 		public override ClaimsPrincipal User { get; set; }
 		public override WebSocketManager WebSockets { get; } = null!;
 
-		public HttpContextStub(Claim roleClaimType)
+		public HttpContextStub(Claim? roleClaimType, IPAddress? remoteIpAddress = null)
 		{
-			User = new ClaimsPrincipal(new ClaimsIdentity(new List<Claim>
-			{
-				roleClaimType
-			}, "TestAuthType"));
+			Connection = new ConnectionInfoStub { RemoteIpAddress = remoteIpAddress };
+			User = new ClaimsPrincipal(new ClaimsIdentity(roleClaimType != null ? [roleClaimType] : [], "TestAuthType"));
 		}
 
-		public HttpContextStub(ClaimsPrincipal user)
+		public HttpContextStub(ClaimsPrincipal user, IPAddress? remoteIpAddress = null)
 		{
+			Connection = new ConnectionInfoStub { RemoteIpAddress = remoteIpAddress };
 			User = user;
 		}
 
diff --git a/Engine/Source/Programs/Horde/HordeServer.Tests.Shared/ServerCallContextStub.cs b/Engine/Source/Programs/Horde/HordeServer.Tests.Shared/ServerCallContextStub.cs
--- a/Engine/Source/Programs/Horde/HordeServer.Tests.Shared/ServerCallContextStub.cs
+++ b/Engine/Source/Programs/Horde/HordeServer.Tests.Shared/ServerCallContextStub.cs
@@ -1,5 +1,6 @@
 // Copyright Epic Games, Inc. All Rights Reserved.
 
+using System.Net;
 using System.Security.Claims;
 using Grpc.Core;
 using HordeServer.Utilities;
@@ -33,26 +34,26 @@ namespace HordeServer.Tests
 			}, "TestAuthType")));
 		}
 
-		public static ServerCallContext ForAdmin()
+		public static ServerCallContext ForAdmin(IPAddress? remoteIpAddress = null)
 		{
 			return new ServerCallContextStub(new ClaimsPrincipal(new ClaimsIdentity(new List<Claim>
 			{
 				HordeClaims.AdminClaim.ToClaim()
-			}, "TestAuthType")));
+			}, "TestAuthType")), remoteIpAddress);
 		}
 
-		public ServerCallContextStub(Claim roleClaimType)
+		public ServerCallContextStub(Claim? roleClaimType, IPAddress? remoteIpAddress = null)
 		{
 			// The GetHttpContext extension falls back to getting the HttpContext from UserState
 			// We can piggyback on that behavior during tests
-			UserState[HttpContextKey] = new HttpContextStub(roleClaimType);
+			UserState[HttpContextKey] = new HttpContextStub(roleClaimType, remoteIpAddress);
 		}
 
-		public ServerCallContextStub(ClaimsPrincipal user)
+		public ServerCallContextStub(ClaimsPrincipal user, IPAddress? remoteIpAddress = null)
 		{
 			// The GetHttpContext extension falls back to getting the HttpContext from UserState
 			// We can piggyback on that behavior during tests
-			UserState[HttpContextKey] = new HttpContextStub(user);
+			UserState[HttpContextKey] = new HttpContextStub(user, remoteIpAddress);
 		}
 
 		public void SetCancellationToken(CancellationToken cancellationToken)
diff --git a/Engine/Source/Programs/Horde/Plugins/Build/HordeServer.Build.Tests/Server/RpcServiceTest.cs b/Engine/Source/Programs/Horde/Plugins/Build/HordeServer.Build.Tests/Server/RpcServiceTest.cs
--- a/Engine/Source/Programs/Horde/Plugins/Build/HordeServer.Build.Tests/Server/RpcServiceTest.cs
+++ b/Engine/Source/Programs/Horde/Plugins/Build/HordeServer.Build.Tests/Server/RpcServiceTest.cs
@@ -2,6 +2,7 @@
 
 using System.IdentityModel.Tokens.Jwt;
 using System.IO.Compression;
+using System.Net;
 using System.Security.Claims;
 using System.Text;
 using EpicGames.Horde.Agents;
@@ -14,107 +15,11 @@ using HordeServer.Acls;
 using HordeServer.Agents;
 using HordeServer.Tools;
 using HordeServer.Utilities;
-using Microsoft.AspNetCore.Http;
-using Microsoft.AspNetCore.Http.Features;
 using Microsoft.Extensions.DependencyInjection;
+using Microsoft.Extensions.Options;
 
 namespace HordeServer.Tests.Server
 {
-	sealed class HttpContextStub : HttpContext
-	{
-		public override ConnectionInfo Connection { get; } = null!;
-		public override IFeatureCollection Features { get; } = null!;
-		public override IDictionary<object, object?> Items { get; set; } = null!;
-		public override HttpRequest Request { get; } = null!;
-		public override CancellationToken RequestAborted { get; set; }
-		public override IServiceProvider RequestServices { get; set; } = null!;
-		public override HttpResponse Response { get; } = null!;
-		public override ISession Session { get; set; } = null!;
-		public override string TraceIdentifier { get; set; } = null!;
-		public override ClaimsPrincipal User { get; set; }
-		public override WebSocketManager WebSockets { get; } = null!;
-
-		public HttpContextStub(Claim? roleClaimType)
-		{
-			User = new ClaimsPrincipal(new ClaimsIdentity(roleClaimType != null ? [roleClaimType] : [], "TestAuthType"));
-		}
-
-		public HttpContextStub(ClaimsPrincipal user)
-		{
-			User = user;
-		}
-
-		public override void Abort()
-		{
-			throw new NotImplementedException();
-		}
-	}
-
-	public class ServerCallContextStub : ServerCallContext
-	{
-		// Copied from ServerCallContextExtensions.cs in Grpc.Core
-		const string HttpContextKey = "__HttpContext";
-
-		protected override string MethodCore { get; } = null!;
-		protected override string HostCore { get; } = null!;
-		protected override string PeerCore { get; } = null!;
-		protected override DateTime DeadlineCore { get; } = DateTime.Now.AddHours(24);
-		protected override Metadata RequestHeadersCore { get; } = null!;
-		protected override CancellationToken CancellationTokenCore => _cancellationToken;
-		protected override Metadata ResponseTrailersCore { get; } = null!;
-		protected override Status StatusCore { get; set; }
-		protected override WriteOptions? WriteOptionsCore { get; set; } = null!;
-		protected override AuthContext AuthContextCore { get; } = null!;
-
-		private CancellationToken _cancellationToken;
-
-		public static ServerCallContext ForAdminWithAgentSessionId(string agentSessionId)
-		{
-			return new ServerCallContextStub(new ClaimsPrincipal(new ClaimsIdentity(new List<Claim>
-			{
-				HordeClaims.AdminClaim.ToClaim(),
-				new Claim(HordeClaimTypes.AgentSessionId, agentSessionId),
-			}, "TestAuthType")));
-		}
-
-		public static ServerCallContext ForAdmin()
-		{
-			return new ServerCallContextStub(new ClaimsPrincipal(new ClaimsIdentity(new List<Claim>
-			{
-				HordeClaims.AdminClaim.ToClaim()
-			}, "TestAuthType")));
-		}
-
-		public ServerCallContextStub(Claim? roleClaimType)
-		{
-			// The GetHttpContext extension falls back to getting the HttpContext from UserState
-			// We can piggyback on that behavior during tests
-			UserState[HttpContextKey] = new HttpContextStub(roleClaimType);
-		}
-
-		public ServerCallContextStub(ClaimsPrincipal user)
-		{
-			// The GetHttpContext extension falls back to getting the HttpContext from UserState
-			// We can piggyback on that behavior during tests
-			UserState[HttpContextKey] = new HttpContextStub(user);
-		}
-
-		public void SetCancellationToken(CancellationToken cancellationToken)
-		{
-			_cancellationToken = cancellationToken;
-		}
-
-		protected override Task WriteResponseHeadersAsyncCore(Metadata responseHeaders)
-		{
-			throw new NotImplementedException();
-		}
-
-		protected override ContextPropagationToken CreatePropagationTokenCore(ContextPropagationOptions? options)
-		{
-			throw new NotImplementedException();
-		}
-	}
-
 	[TestClass]
 	public class RpcServiceTest : BuildTestSetup
 	{
@@ -344,6 +249,53 @@ namespace HordeServer.Tests.Server
 			}
 		}
 		
+		[TestMethod]
+		public async Task UpdateSession_SessionMismatch_QuarantinesAfterThresholdAsync()
+		{
+			int threshold = ServiceProvider.GetRequiredService<IOptionsSnapshot<ComputeServerConfig>>().Value.SessionConflictThreshold;
+			AgentId agentId = new AgentId("conflict-agent");
+			IPAddress callerIp = IPAddress.Parse("10.0.0.1");
+			ServerCallContext adminCtx = ServerCallContextStub.ForAdmin(callerIp);
+
+			// Create agent and first session
+			RpcCreateSessionRequest createReq1 = new RpcCreateSessionRequest
+			{
+				Id = agentId.ToString(),
+				Capabilities = new RpcAgentCapabilities()
+			};
+			RpcCreateSessionResponse createRes1 = await RpcService.CreateSession(createReq1, adminCtx);
+			string staleSessionId = createRes1.SessionId;
+
+			// Create a second session (invalidates the first)
+			RpcCreateSessionRequest createReq2 = new RpcCreateSessionRequest
+			{
+				Id = agentId.ToString(),
+				Capabilities = new RpcAgentCapabilities()
+			};
+			await RpcService.CreateSession(createReq2, adminCtx);
+
+			// Attempt UpdateSession with the stale session ID repeatedly
+			// Each call should hit the session mismatch path and record a conflict
+			for (int i = 0; i < threshold; i++)
+			{
+				TestAsyncStreamReader<RpcUpdateSessionRequest> requestStream = new(adminCtx);
+				TestServerStreamWriter<RpcUpdateSessionResponse> responseStream = new(adminCtx);
+				requestStream.AddMessage(new RpcUpdateSessionRequest { AgentId = agentId.ToString(), SessionId = staleSessionId });
+
+				StructuredRpcException ex = await Assert.ThrowsExactlyAsync<StructuredRpcException>(
+					() => RpcService.UpdateSession(requestStream, responseStream, adminCtx));
+				Assert.AreEqual(StatusCode.PermissionDenied, ex.StatusCode);
+			}
+
+			// Verify the agent has been quarantined
+			IAgent? agent = await AgentService.GetAgentAsync(agentId);
+			Assert.IsNotNull(agent);
+			Assert.IsFalse(agent.Enabled, "Agent should be disabled after session conflict threshold");
+			Assert.IsNotNull(agent.Comment);
+			Assert.IsTrue(agent.Comment.StartsWith(SessionConflictService.CommentPrefix, StringComparison.Ordinal),
+				$"Comment should start with session-conflict prefix, was: {agent.Comment}");
+		}
+
 		private static ServerCallContextStub CreateServerContext(string? roleClaim)
 		{
 			return new ServerCallContextStub(roleClaim != null ? new AclClaimConfig(HordeClaimTypes.Role, roleClaim).ToClaim() : null);
diff --git a/Engine/Source/Programs/Horde/Plugins/Build/HordeServer.Build/Agents/AgentReportService.cs b/Engine/Source/Programs/Horde/Plugins/Build/HordeServer.Build/Agents/AgentReportService.cs
--- a/Engine/Source/Programs/Horde/Plugins/Build/HordeServer.Build/Agents/AgentReportService.cs
+++ b/Engine/Source/Programs/Horde/Plugins/Build/HordeServer.Build/Agents/AgentReportService.cs
@@ -25,6 +25,11 @@ namespace HordeServer.Agents
 		/// List of agents stuck in a upgrade loop, plus the number of times they've attempted to upgrade
 		/// </summary>
 		public List<(AgentId, int)> UpgradeLoop { get; } = new List<(AgentId, int)>();
+
+		/// <summary>
+		/// List of agents quarantined due to session conflicts, with their comment text
+		/// </summary>
+		public List<(AgentId Id, string Comment)> SessionConflicts { get; } = new List<(AgentId, string)>();
 	}
 
 	[SingletonDocument("agent-report-state", "6268871c211d05611b3e4fd8")]
@@ -99,6 +104,11 @@ namespace HordeServer.Agents
 							report.UpgradeLoop.Add((agent.Id, agent.UpgradeAttemptCount.Value));
 						}
 					}
+
+					if (!agent.Enabled && agent.Comment?.StartsWith(SessionConflictService.CommentPrefix, StringComparison.Ordinal) == true)
+					{
+						report.SessionConflicts.Add((agent.Id, agent.Comment));
+					}
 				}
 
 				await _notificationService.SendAgentReportAsync(report, cancellationToken);
diff --git a/Engine/Source/Programs/Horde/Plugins/Compute/HordeServer.Compute.Tests/Agents/SessionConflictServiceTest.cs b/Engine/Source/Programs/Horde/Plugins/Compute/HordeServer.Compute.Tests/Agents/SessionConflictServiceTest.cs
new file mode 100644
--- /dev/null
+++ b/Engine/Source/Programs/Horde/Plugins/Compute/HordeServer.Compute.Tests/Agents/SessionConflictServiceTest.cs
@@ -0,0 +1,136 @@
+// Copyright Epic Games, Inc. All Rights Reserved.
+
+using System.Net;
+using EpicGames.Horde.Agents;
+using HordeServer.Agents.Sessions;
+using HordeServer.Agents;
+using HordeServer.Server;
+using Microsoft.Extensions.DependencyInjection;
+using Microsoft.Extensions.Logging;
+using Microsoft.Extensions.Options;
+using StackExchange.Redis;
+
+namespace HordeServer.Tests.Agents
+{
+	[TestClass]
+	public class SessionConflictServiceTest : ComputeTestSetup
+	{
+		private readonly IRedisService _redisService;
+
+		public SessionConflictServiceTest()
+		{
+			_redisService = GetRedisServiceSingleton();
+		}
+
+		private SessionConflictService CreateService(ComputeServerConfig? config = null)
+		{
+			IOptionsMonitor<ComputeServerConfig> configMonitor = config != null
+				? new TestOptionsMonitor<ComputeServerConfig>(config)
+				: ServiceProvider.GetRequiredService<IOptionsMonitor<ComputeServerConfig>>();
+			ILogger<SessionConflictService> logger = ServiceProvider.GetRequiredService<ILogger<SessionConflictService>>();
+			return new SessionConflictService(_redisService, AgentCollection, configMonitor, logger);
+		}
+
+		private static readonly IPAddress s_ip1 = IPAddress.Parse("10.1.2.3");
+		private static readonly IPAddress s_ip2 = IPAddress.Parse("10.1.2.4");
+
+		[TestMethod]
+		public async Task BelowThreshold_NoQuarantineAsync()
+		{
+			SessionConflictService service = CreateService();
+			IAgent agent = await AgentCollection.AddAsync(new CreateAgentOptions(new AgentId("conflict-test-1"), AgentMode.Dedicated, false, ""));
+
+			// Record 2 conflicts (threshold is 3)
+			await service.RecordConflictAsync(agent.Id, SessionIdUtils.GenerateNewId(), SessionIdUtils.GenerateNewId(), s_ip1, CancellationToken.None);
+			await service.RecordConflictAsync(agent.Id, SessionIdUtils.GenerateNewId(), SessionIdUtils.GenerateNewId(), s_ip2, CancellationToken.None);
+
+			IAgent? updated = await AgentCollection.GetAsync(agent.Id);
+			Assert.IsNotNull(updated);
+			Assert.IsTrue(updated.Enabled, "Agent should still be enabled below threshold");
+		}
+
+		[TestMethod]
+		public async Task AtThreshold_QuarantinesAgentAsync()
+		{
+			SessionConflictService service = CreateService();
+			IAgent agent = await AgentCollection.AddAsync(new CreateAgentOptions(new AgentId("conflict-test-2"), AgentMode.Dedicated, false, ""));
+
+			// Record 3 conflicts (hits threshold)
+			await service.RecordConflictAsync(agent.Id, SessionIdUtils.GenerateNewId(), SessionIdUtils.GenerateNewId(), s_ip1, CancellationToken.None);
+			await service.RecordConflictAsync(agent.Id, SessionIdUtils.GenerateNewId(), SessionIdUtils.GenerateNewId(), s_ip2, CancellationToken.None);
+			await service.RecordConflictAsync(agent.Id, SessionIdUtils.GenerateNewId(), SessionIdUtils.GenerateNewId(), s_ip1, CancellationToken.None);
+
+			IAgent? updated = await AgentCollection.GetAsync(agent.Id);
+			Assert.IsNotNull(updated);
+			Assert.IsFalse(updated.Enabled, "Agent should be disabled after reaching threshold");
+			Assert.IsNotNull(updated.Comment);
+			Assert.IsTrue(updated.Comment.StartsWith(SessionConflictService.CommentPrefix, StringComparison.Ordinal), $"Comment should start with prefix, was: {updated.Comment}");
+			Assert.IsTrue(updated.Comment.Contains("10.1.2.3", StringComparison.Ordinal), "Comment should contain IP 1");
+			Assert.IsTrue(updated.Comment.Contains("10.1.2.4", StringComparison.Ordinal), "Comment should contain IP 2");
+		}
+
+		[TestMethod]
+		public async Task WindowExpiry_OldEventsNotCountedAsync()
+		{
+			ComputeServerConfig config = new ComputeServerConfig { SessionConflictWindowSeconds = 300 };
+			SessionConflictService service = CreateService(config);
+			IAgent agent = await AgentCollection.AddAsync(new CreateAgentOptions(new AgentId("conflict-test-3"), AgentMode.Dedicated, false, ""));
+
+			// Insert 2 events with timestamps well outside the window (10 minutes old)
+			IDatabase db = _redisService.GetDatabase();
+			RedisKey key = SessionConflictService.GetRedisKey(agent.Id);
+			long expiredTicks = DateTimeOffset.UtcNow.AddSeconds(-600).Ticks;
+			await db.SortedSetAddAsync(key, $"old1:{Guid.NewGuid():N}", expiredTicks);
+			await db.SortedSetAddAsync(key, $"old2:{Guid.NewGuid():N}", expiredTicks);
+
+			// Record 1 more conflict via the service (triggers pruning of old entries)
+			await service.RecordConflictAsync(agent.Id, SessionIdUtils.GenerateNewId(), SessionIdUtils.GenerateNewId(), s_ip1, CancellationToken.None);
+
+			// Only 1 event should remain in the window (below threshold of 3)
+			long count = await db.SortedSetLengthAsync(key);
+			Assert.AreEqual(1, count, "Only the fresh event should remain after pruning");
+
+			IAgent? updated = await AgentCollection.GetAsync(agent.Id);
+			Assert.IsNotNull(updated);
+			Assert.IsTrue(updated.Enabled, "Agent should still be enabled when old events were pruned");
+		}
+
+		[TestMethod]
+		public async Task AlreadyDisabled_IdempotentAsync()
+		{
+			SessionConflictService service = CreateService();
+			IAgent agent = await AgentCollection.AddAsync(new CreateAgentOptions(new AgentId("conflict-test-4"), AgentMode.Dedicated, false, ""));
+
+			// Pre-disable the agent with a custom comment
+			agent = (await agent.TryUpdateAsync(new UpdateAgentOptions(Enabled: false, Comment: "manually disabled")))!;
+			Assert.IsFalse(agent.Enabled);
+
+			// Record 3+ conflicts (would normally trigger quarantine)
+			await service.RecordConflictAsync(agent.Id, SessionIdUtils.GenerateNewId(), SessionIdUtils.GenerateNewId(), s_ip1, CancellationToken.None);
+			await service.RecordConflictAsync(agent.Id, SessionIdUtils.GenerateNewId(), SessionIdUtils.GenerateNewId(), s_ip2, CancellationToken.None);
+			await service.RecordConflictAsync(agent.Id, SessionIdUtils.GenerateNewId(), SessionIdUtils.GenerateNewId(), s_ip1, CancellationToken.None);
+
+			IAgent? updated = await AgentCollection.GetAsync(agent.Id);
+			Assert.IsNotNull(updated);
+			Assert.IsFalse(updated.Enabled);
+			Assert.AreEqual("manually disabled", updated.Comment, "Comment should be unchanged when agent was already disabled");
+		}
+
+		[TestMethod]
+		public async Task Disabled_NoRecordingAsync()
+		{
+			ComputeServerConfig disabledConfig = new ComputeServerConfig { SessionConflictEnabled = false };
+			SessionConflictService service = CreateService(disabledConfig);
+			IAgent agent = await AgentCollection.AddAsync(new CreateAgentOptions(new AgentId("conflict-test-5"), AgentMode.Dedicated, false, ""));
+
+			// Record 3+ conflicts with feature disabled
+			await service.RecordConflictAsync(agent.Id, SessionIdUtils.GenerateNewId(), SessionIdUtils.GenerateNewId(), s_ip1, CancellationToken.None);
+			await service.RecordConflictAsync(agent.Id, SessionIdUtils.GenerateNewId(), SessionIdUtils.GenerateNewId(), s_ip2, CancellationToken.None);
+			await service.RecordConflictAsync(agent.Id, SessionIdUtils.GenerateNewId(), SessionIdUtils.GenerateNewId(), s_ip1, CancellationToken.None);
+
+			IAgent? updated = await AgentCollection.GetAsync(agent.Id);
+			Assert.IsNotNull(updated);
+			Assert.IsTrue(updated.Enabled, "Agent should remain enabled when feature is disabled");
+		}
+	}
+}
diff --git a/Engine/Source/Programs/Horde/Plugins/Compute/HordeServer.Compute/Agents/SessionConflictService.cs b/Engine/Source/Programs/Horde/Plugins/Compute/HordeServer.Compute/Agents/SessionConflictService.cs
new file mode 100644
--- /dev/null
+++ b/Engine/Source/Programs/Horde/Plugins/Compute/HordeServer.Compute/Agents/SessionConflictService.cs
@@ -0,0 +1,131 @@
+// Copyright Epic Games, Inc. All Rights Reserved.
+
+using System.Net;
+using EpicGames.Horde.Agents;
+using EpicGames.Horde.Agents.Sessions;
+using HordeServer.Server;
+using Microsoft.Extensions.Logging;
+using Microsoft.Extensions.Options;
+using StackExchange.Redis;
+
+namespace HordeServer.Agents
+{
+	/// <summary>
+	/// Detects session conflict flip-flops (e.g. cloned VMs sharing an agent ID)
+	/// and auto-quarantines the agent after a configurable threshold.
+	/// </summary>
+	public class SessionConflictService
+	{
+		/// <summary>
+		/// Prefix used in agent comments when quarantining due to session conflicts.
+		/// Also used by AgentReportService to identify these agents.
+		/// </summary>
+		public const string CommentPrefix = "[session-conflict]";
+
+		readonly IRedisService _redisService;
+		readonly IAgentCollection _agentCollection;
+		readonly IOptionsMonitor<ComputeServerConfig> _serverConfig;
+		readonly ILogger _logger;
+
+		/// <summary>
+		/// Constructor
+		/// </summary>
+		public SessionConflictService(IRedisService redisService, IAgentCollection agentCollection, IOptionsMonitor<ComputeServerConfig> serverConfig, ILogger<SessionConflictService> logger)
+		{
+			_redisService = redisService;
+			_agentCollection = agentCollection;
+			_serverConfig = serverConfig;
+			_logger = logger;
+		}
+
+		internal static RedisKey GetRedisKey(AgentId agentId) => new RedisKey($"session-conflict:{agentId}");
+
+		// Member format: "{staleSessionId}:{currentSessionId}:{sanitizedIp}:{uniqueId}"
+		// SessionId is a hex BinaryId (no colons), and IPv6 colons are replaced with underscores.
+		static string CreateMember(SessionId staleSessionId, SessionId? currentSessionId, IPAddress? remoteIp)
+		{
+			string sanitizedIp = remoteIp?.ToString().Replace(':', '_') ?? "unknown";
+			return $"{staleSessionId}:{currentSessionId}:{sanitizedIp}:{Guid.NewGuid():N}";
+		}
+
+		static string ParseIpFromMember(string member)
+		{
+			string[] parts = member.Split(':');
+			return parts.Length >= 3 ? parts[2].Replace('_', ':') : "unknown";
+		}
+
+		/// <summary>
+		/// Records a session conflict event and quarantines the agent if the threshold is exceeded.
+		/// </summary>
+		/// <param name="agentId">The agent experiencing the conflict</param>
+		/// <param name="staleSessionId">The session ID that the caller tried to use</param>
+		/// <param name="currentSessionId">The actual current session ID on the agent</param>
+		/// <param name="remoteIp">IP address of the caller</param>
+		/// <param name="cancellationToken">Cancellation token</param>
+		public async Task RecordConflictAsync(AgentId agentId, SessionId staleSessionId, SessionId? currentSessionId, IPAddress? remoteIp, CancellationToken cancellationToken)
+		{
+			ComputeServerConfig config = _serverConfig.CurrentValue;
+			if (!config.SessionConflictEnabled)
+			{
+				return;
+			}
+
+			IDatabase db = _redisService.GetDatabase();
+			RedisKey key = GetRedisKey(agentId);
+			TimeSpan window = TimeSpan.FromSeconds(config.SessionConflictWindowSeconds);
+
+			long nowTicks = DateTimeOffset.UtcNow.Ticks;
+			string member = CreateMember(staleSessionId, currentSessionId, remoteIp);
+
+			// Add event and prune old entries
+			await db.SortedSetAddAsync(key, member, nowTicks);
+			await db.SortedSetRemoveRangeByScoreAsync(key, double.NegativeInfinity, nowTicks - window.Ticks);
+			TimeSpan ttl = window + window;
+			await db.KeyExpireAsync(key, ttl);
+
+			long count = await db.SortedSetLengthAsync(key);
+			if (count < config.SessionConflictThreshold)
+			{
+				return;
+			}
+
+			// Collect distinct IPs from recent events for the quarantine comment
+			SortedSetEntry[] entries = await db.SortedSetRangeByScoreWithScoresAsync(key);
+			HashSet<string> ips = new HashSet<string>(StringComparer.Ordinal);
+			foreach (SortedSetEntry entry in entries)
+			{
+				ips.Add(ParseIpFromMember(entry.Element.ToString()));
+			}
+
+			string ipList = String.Join(", ", ips.Order());
+			string comment = $"{CommentPrefix} Auto-quarantined: session flip-flop detected ({count} conflicts). IPs: {ipList}. Time: {DateTimeOffset.UtcNow:O}. Re-enable after investigating duplicate agent IDs.";
+
+			// Quarantine the agent with retry for optimistic concurrency
+			for (int attempt = 0; attempt < 5; attempt++)
+			{
+				IAgent? agent = await _agentCollection.GetAsync(agentId, cancellationToken);
+				if (agent == null)
+				{
+					break;
+				}
+
+				if (!agent.Enabled)
+				{
+					break; // Already disabled, nothing to do
+				}
+
+				IAgent? updated = await agent.TryUpdateAsync(new UpdateAgentOptions(Enabled: false, Comment: comment), cancellationToken);
+				if (updated != null)
+				{
+					_logger.LogWarning("Agent {AgentId} quarantined due to session conflict flip-flop ({Count} conflicts in {Window}s). IPs: {IpList}", agentId, count, config.SessionConflictWindowSeconds, ipList);
+
+					// Clean up Redis key after successful quarantine
+					await db.KeyDeleteAsync(key);
+					return;
+				}
+			}
+
+			_logger.LogWarning("Failed to quarantine agent {AgentId} after 5 attempts", agentId);
+		}
+	}
+}
diff --git a/Engine/Source/Programs/Horde/Plugins/Compute/HordeServer.Compute/ComputePlugin.cs b/Engine/Source/Programs/Horde/Plugins/Compute/HordeServer.Compute/ComputePlugin.cs
--- a/Engine/Source/Programs/Horde/Plugins/Compute/HordeServer.Compute/ComputePlugin.cs
+++ b/Engine/Source/Programs/Horde/Plugins/Compute/HordeServer.Compute/ComputePlugin.cs
@@ -106,6 +106,7 @@ namespace HordeServer
 
 			services.AddSingleton<ComputeService>();
 			services.AddSingleton<EnrollmentService>();
+			services.AddSingleton<SessionConflictService>();
 			services.AddSingleton<LogTailService>();
 			services.AddSingleton<PoolService>();
 
diff --git a/Engine/Source/Programs/Horde/Plugins/Compute/HordeServer.Compute/ComputeServerConfig.cs b/Engine/Source/Programs/Horde/Plugins/Compute/HordeServer.Compute/ComputeServerConfig.cs
--- a/Engine/Source/Programs/Horde/Plugins/Compute/HordeServer.Compute/ComputeServerConfig.cs
+++ b/Engine/Source/Programs/Horde/Plugins/Compute/HordeServer.Compute/ComputeServerConfig.cs
@@ -67,6 +67,23 @@ namespace HordeServer
 		/// </summary>
 		public int AgentPoolScaleInCooldownSeconds { get; set; } = 1200; // 20 mins
 
+		/// <summary>
+		/// Whether to enable session conflict detection and auto-quarantine.
+		/// When multiple machines accidentally connect using the same agent ID (e.g. cloned VMs),
+		/// repeated session mismatches will cause the agent to be automatically disabled.
+		/// </summary>
+		public bool SessionConflictEnabled { get; set; } = true;
+
+		/// <summary>
+		/// Number of session mismatches within the sliding window required to trigger auto-quarantine
+		/// </summary>
+		public int SessionConflictThreshold { get; set; } = 3;
+
+		/// <summary>
+		/// Sliding window in seconds for counting session conflict events
+		/// </summary>
+		public int SessionConflictWindowSeconds { get; set; } = 300;
+
 		/// <summary>
 		/// Port to listen on for tunneling compute sockets to agents
 		/// </summary>
diff --git a/Engine/Source/Programs/Horde/Plugins/Compute/HordeServer.Compute/Server/RpcService.cs b/Engine/Source/Programs/Horde/Plugins/Compute/HordeServer.Compute/Server/RpcService.cs
--- a/Engine/Source/Programs/Horde/Plugins/Compute/HordeServer.Compute/Server/RpcService.cs
+++ b/Engine/Source/Programs/Horde/Plugins/Compute/HordeServer.Compute/Server/RpcService.cs
@@ -43,6 +43,8 @@ namespace HordeServer.Server
 		readonly IAgentTelemetryCollection _agentTelemetryCollection;
 		readonly IAclService _aclService;
 		readonly IOptionsSnapshot<ComputeConfig> _computeConfig;
+		readonly SessionConflictService _sessionConflictService;
+		readonly IOptionsMonitor<ComputeServerConfig> _staticComputeConfig;
 		readonly ILogger _logger;
 		readonly ITelemetryWriter _telemetryWriter;
 
@@ -50,12 +52,14 @@ namespace HordeServer.Server
 		/// Constructor
 		/// </summary>
 		public RpcService(
-			AgentService agentService, 
-			ILifetimeService lifetimeService, 
-			IToolCollection toolCollection, 
-			IAgentTelemetryCollection agentTelemetryCollection, 
-			IAclService aclService, 
-			IOptionsSnapshot<ComputeConfig> computeConfig, 
+			AgentService agentService,
+			ILifetimeService lifetimeService,
+			IToolCollection toolCollection,
+			IAgentTelemetryCollection agentTelemetryCollection,
+			IAclService aclService,
+			IOptionsSnapshot<ComputeConfig> computeConfig,
+			SessionConflictService sessionConflictService,
+			IOptionsMonitor<ComputeServerConfig> staticComputeConfig,
 			ILogger<RpcService> logger,
 			ITelemetryWriter telemetryWriter)
 		{
@@ -65,6 +69,8 @@ namespace HordeServer.Server
 			_agentTelemetryCollection = agentTelemetryCollection;
 			_aclService = aclService;
 			_computeConfig = computeConfig;
+			_sessionConflictService = sessionConflictService;
+			_staticComputeConfig = staticComputeConfig;
 			_logger = logger;
 			_telemetryWriter = telemetryWriter;
 		}
@@ -258,6 +264,9 @@ namespace HordeServer.Server
 				throw new StructuredRpcException(StatusCode.NotFound, "Agent {AgentId} not found", agentId);
 			}
 
+			IPAddress? createRemoteIp = context.GetHttpContext().Connection?.RemoteIpAddress;
+			_logger.LogDebug("Session {SessionId} created for agent {AgentId} from IP {RemoteIp}", agent.SessionId, agent.Id, createRemoteIp);
+
 			// Create the response
 			RpcCreateSessionResponse response = new RpcCreateSessionResponse();
 			response.AgentId = agent.Id.ToString();
@@ -309,6 +318,22 @@ namespace HordeServer.Server
 						// Check we're authorized to update it
 						if (agent.SessionId != sessionId)
 						{
+							IPAddress? remoteIp = context.GetHttpContext().Connection?.RemoteIpAddress;
+							bool conflictEnabled = _staticComputeConfig.CurrentValue.SessionConflictEnabled;
+							bool alreadyQuarantined = conflictEnabled && !agent.Enabled && agent.Comment?.StartsWith(SessionConflictService.CommentPrefix, StringComparison.Ordinal) == true;
+
+							if (alreadyQuarantined)
+							{
+								// Suppress log spam: log at Debug and throw plain RpcException (bypasses the interceptor's ERROR log).
+								_logger.LogDebug("Session mismatch for already-quarantined agent {AgentId}: caller has {StaleSessionId} but current is {CurrentSessionId}, remote IP {RemoteIp}", agent.Id, sessionId, agent.SessionId, remoteIp);
+								throw new RpcException(new Status(StatusCode.PermissionDenied, $"Agent {request.AgentId} is quarantined for session conflicts."));
+							}
+
+							_logger.LogInformation("Session mismatch for agent {AgentId}: caller has {StaleSessionId} but current is {CurrentSessionId}, remote IP {RemoteIp}", agent.Id, sessionId, agent.SessionId, remoteIp);
+							if (conflictEnabled)
+							{
+								await _sessionConflictService.RecordConflictAsync(agent.Id, sessionId, agent.SessionId, remoteIp, context.CancellationToken);
+							}
 							throw new StructuredRpcException(StatusCode.PermissionDenied, "Agent {AgentId} has completed session {SessionId}; now executing session {NewSessionId}. Cannot update state.", request.AgentId, sessionId, agent.SessionId?.ToString() ?? "(None)");
 						}
 						if (!_agentService.AuthorizeSession(agent, context.GetHttpContext().User, out string authReason))
