Horde: Invite users nominated to fix build health issue to Slack triage channel
Add InviteUsersAsync call after posting nomination message in CreateOrUpdateWorkflowThreadAsync, ensuring the user receives a notification.
Suppress already_in_channel warning in non-admin invite path, matching existing admin path behavior, to reduce noise in logs.

#Synced-CL 0

diff --git a/Engine/Source/Programs/Horde/Plugins/Build/HordeServer.Build/Notifications/Sinks/SlackNotificationSink.cs b/Engine/Source/Programs/Horde/Plugins/Build/HordeServer.Build/Notifications/Sinks/SlackNotificationSink.cs
--- a/Engine/Source/Programs/Horde/Plugins/Build/HordeServer.Build/Notifications/Sinks/SlackNotificationSink.cs
+++ b/Engine/Source/Programs/Horde/Plugins/Build/HordeServer.Build/Notifications/Sinks/SlackNotificationSink.cs
@@ -1035,7 +1035,7 @@ namespace HordeServer.Notifications.Sinks
 							}
 						}
 					}
-					else
+					else if (!errorCode.Equals("already_in_channel", StringComparison.Ordinal))
 					{
 						_logger.LogWarning("Unable to invite user {UserId} to {Channel}: {Error}", slackUserId, channel, errorCode);
 					}
@@ -1415,6 +1415,12 @@ namespace HordeServer.Notifications.Sinks
 					string assignmentEventId = $"issue_{issue.Id}_nominated";
 					string assignmentMessage = $"{await FormatMentionAsync(issue.OwnerId.Value, workflow.AllowMentions, cancellationToken)} was nominated to fix by {await FormatNameAsync(issue.NominatedById.Value, cancellationToken)}.";
 					await PostSingleMessageToThreadAsync(triageChannel, assignmentEventId, threadId, assignmentMessage, cancellationToken);
+
+					// Invite the nominated user to the channel so they receive the notification
+					if (_environment.IsProduction() && workflow.AllowMentions)
+					{
+						await InviteUsersAsync(state.Channel, new[] { issue.OwnerId.Value }, workflow.InviteRestrictedUsers, cancellationToken);
+					}
 				}
 			}
 
