Horde: Updating schema documentation

#rnx

diff --git a/Engine/Source/Programs/Horde/Docs/Config/Schema/AclActions.md b/Engine/Source/Programs/Horde/Docs/Config/Schema/AclActions.md
--- a/Engine/Source/Programs/Horde/Docs/Config/Schema/AclActions.md
+++ b/Engine/Source/Programs/Horde/Docs/Config/Schema/AclActions.md
@@ -37,47 +37,43 @@
 | `DeleteAccount` |  |
 | `ViewAccount` |  |
 
-## Artifacts
-
-| Name | Description |
-| ---- | ----------- |
-| `ReadArtifact` |  |
-| `WriteArtifact` |  |
-| `DeleteArtifact` |  |
-| `UploadArtifact` |  |
-| `DownloadArtifact` |  |
-
-## Bisect
+## Agents
 
 | Name | Description |
 | ---- | ----------- |
-| `CreateBisectTask` |  |
-| `UpdateBisectTask` |  |
-| `ViewBisectTask` |  |
+| `CreateAgent` |  |
+| `CreateWorkstationAgent` |  |
+| `UpdateAgent` |  |
+| `DeleteAgent` |  |
+| `ViewAgent` |  |
+| `ListAgents` |  |
 
-## BuildHealth
+## Compute
 
 | Name | Description |
 | ---- | ----------- |
-| `AddUpdateDeleteBuildHealthFilter` |  |
+| `AddComputeTasks` |  |
+| `GetComputeTasks` |  |
+| `UbaCacheRead` |  |
+| `UbaCacheWrite` |  |
 
-## Devices
+## Leases
 
 | Name | Description |
 | ---- | ----------- |
-| `DeviceRead` |  |
-| `DeviceWrite` |  |
+| `ViewLeases` |  |
+| `ViewLeaseTasks` |  |
 
-## Jobs
+## Logs
 
 | Name | Description |
 | ---- | ----------- |
-| `CreateJob` |  |
-| `UpdateJob` |  |
-| `DeleteJob` |  |
-| `ExecuteJob` |  |
-| `RetryJobStep` |  |
-| `ViewJob` |  |
+| `CreateLog` |  |
+| `UpdateLog` |  |
+| `ViewLog` |  |
+| `WriteLogData` |  |
+| `CreateEvent` |  |
+| `ViewEvent` |  |
 
 ## Notices
 
@@ -87,27 +83,15 @@
 | `UpdateNotice` |  |
 | `DeleteNotice` |  |
 
-## Notifications
-
-| Name | Description |
-| ---- | ----------- |
-| `CreateSubscription` |  |
-
-## Projects
-
-| Name | Description |
-| ---- | ----------- |
-| `CreateProject` |  |
-| `DeleteProject` |  |
-| `UpdateProject` |  |
-| `ViewProject` |  |
-
-## Replicators
+## Pools
 
 | Name | Description |
 | ---- | ----------- |
-| `UpdateReplicator` |  |
-| `ViewReplicator` |  |
+| `CreatePool` |  |
+| `UpdatePool` |  |
+| `DeletePool` |  |
+| `ViewPool` |  |
+| `ListPools` |  |
 
 ## ServiceAccounts
 
@@ -118,23 +102,20 @@
 | `DeleteAccount` |  |
 | `ViewAccount` |  |
 
-## Streams
+## Sessions
 
 | Name | Description |
 | ---- | ----------- |
-| `CreateStream` |  |
-| `UpdateStream` |  |
-| `DeleteStream` |  |
-| `ViewStream` |  |
-| `ViewChanges` |  |
-| `ViewTemplate` |  |
+| `CreateSession` |  |
+| `ViewSession` |  |
 
-## TestData
+## Software
 
 | Name | Description |
 | ---- | ----------- |
-| `TestHealthRead` |  |
-| `TestHealthWrite` |  |
+| `UploadSoftware` |  |
+| `DownloadSoftware` |  |
+| `DeleteSoftware` |  |
 
 ## Accounts
 
@@ -156,30 +137,6 @@
 | `ViewAgent` |  |
 | `ListAgents` |  |
 
-## Artifacts
-
-| Name | Description |
-| ---- | ----------- |
-| `ReadArtifact` |  |
-| `WriteArtifact` |  |
-| `DeleteArtifact` |  |
-| `UploadArtifact` |  |
-| `DownloadArtifact` |  |
-
-## Bisect
-
-| Name | Description |
-| ---- | ----------- |
-| `CreateBisectTask` |  |
-| `UpdateBisectTask` |  |
-| `ViewBisectTask` |  |
-
-## BuildHealth
-
-| Name | Description |
-| ---- | ----------- |
-| `AddUpdateDeleteBuildHealthFilter` |  |
-
 ## Compute
 
 | Name | Description |
@@ -189,24 +146,6 @@
 | `UbaCacheRead` |  |
 | `UbaCacheWrite` |  |
 
-## Devices
-
-| Name | Description |
-| ---- | ----------- |
-| `DeviceRead` |  |
-| `DeviceWrite` |  |
-
-## Jobs
-
-| Name | Description |
-| ---- | ----------- |
-| `CreateJob` |  |
-| `UpdateJob` |  |
-| `DeleteJob` |  |
-| `ExecuteJob` |  |
-| `RetryJobStep` |  |
-| `ViewJob` |  |
-
 ## Leases
 
 | Name | Description |
@@ -233,12 +172,6 @@
 | `UpdateNotice` |  |
 | `DeleteNotice` |  |
 
-## Notifications
-
-| Name | Description |
-| ---- | ----------- |
-| `CreateSubscription` |  |
-
 ## Pools
 
 | Name | Description |
@@ -249,22 +182,6 @@
 | `ViewPool` |  |
 | `ListPools` |  |
 
-## Projects
-
-| Name | Description |
-| ---- | ----------- |
-| `CreateProject` |  |
-| `DeleteProject` |  |
-| `UpdateProject` |  |
-| `ViewProject` |  |
-
-## Replicators
-
-| Name | Description |
-| ---- | ----------- |
-| `UpdateReplicator` |  |
-| `ViewReplicator` |  |
-
 ## ServiceAccounts
 
 | Name | Description |
@@ -289,24 +206,6 @@
 | `DownloadSoftware` |  |
 | `DeleteSoftware` |  |
 
-## Streams
-
-| Name | Description |
-| ---- | ----------- |
-| `CreateStream` |  |
-| `UpdateStream` |  |
-| `DeleteStream` |  |
-| `ViewStream` |  |
-| `ViewChanges` |  |
-| `ViewTemplate` |  |
-
-## TestData
-
-| Name | Description |
-| ---- | ----------- |
-| `TestHealthRead` |  |
-| `TestHealthWrite` |  |
-
 ## Accounts
 
 | Name | Description |
@@ -802,19 +701,6 @@
 | `DownloadSoftware` |  |
 | `DeleteSoftware` |  |
 
-## Storage
-
-| Name | Description |
-| ---- | ----------- |
-| `ReadBlobs` |  |
-| `WriteBlobs` |  |
-| `ReadAliases` |  |
-| `WriteAliases` |  |
-| `DeleteAliases` |  |
-| `ReadRefs` |  |
-| `WriteRefs` |  |
-| `DeleteRefs` |  |
-
 ## Streams
 
 | Name | Description |
@@ -986,19 +872,6 @@
 | `DownloadSoftware` |  |
 | `DeleteSoftware` |  |
 
-## Storage
-
-| Name | Description |
-| ---- | ----------- |
-| `ReadBlobs` |  |
-| `WriteBlobs` |  |
-| `ReadAliases` |  |
-| `WriteAliases` |  |
-| `DeleteAliases` |  |
-| `ReadRefs` |  |
-| `WriteRefs` |  |
-| `DeleteRefs` |  |
-
 ## Streams
 
 | Name | Description |
@@ -1010,6 +883,12 @@
 | `ViewChanges` |  |
 | `ViewTemplate` |  |
 
+## Symbols
+
+| Name | Description |
+| ---- | ----------- |
+| `ReadSymbols` |  |
+
 ## TestData
 
 | Name | Description |
@@ -1170,19 +1049,6 @@
 | `DownloadSoftware` |  |
 | `DeleteSoftware` |  |
 
-## Storage
-
-| Name | Description |
-| ---- | ----------- |
-| `ReadBlobs` |  |
-| `WriteBlobs` |  |
-| `ReadAliases` |  |
-| `WriteAliases` |  |
-| `DeleteAliases` |  |
-| `ReadRefs` |  |
-| `WriteRefs` |  |
-| `DeleteRefs` |  |
-
 ## Streams
 
 | Name | Description |
@@ -1207,6 +1073,13 @@
 | `TestHealthRead` |  |
 | `TestHealthWrite` |  |
 
+## Tools
+
+| Name | Description |
+| ---- | ----------- |
+| `DownloadTool` |  |
+| `UploadTool` |  |
+
 ## Accounts
 
 | Name | Description |
@@ -1336,6 +1209,12 @@
 | `UpdateReplicator` |  |
 | `ViewReplicator` |  |
 
+## Secrets
+
+| Name | Description |
+| ---- | ----------- |
+| `ViewSecret` |  |
+
 ## ServiceAccounts
 
 | Name | Description |
@@ -1360,19 +1239,6 @@
 | `DownloadSoftware` |  |
 | `DeleteSoftware` |  |
 
-## Storage
-
-| Name | Description |
-| ---- | ----------- |
-| `ReadBlobs` |  |
-| `WriteBlobs` |  |
-| `ReadAliases` |  |
-| `WriteAliases` |  |
-| `DeleteAliases` |  |
-| `ReadRefs` |  |
-| `WriteRefs` |  |
-| `DeleteRefs` |  |
-
 ## Streams
 
 | Name | Description |
diff --git a/Engine/Source/Programs/Horde/Docs/Config/Schema/Globals.md b/Engine/Source/Programs/Horde/Docs/Config/Schema/Globals.md
--- a/Engine/Source/Programs/Horde/Docs/Config/Schema/Globals.md
+++ b/Engine/Source/Programs/Horde/Docs/Config/Schema/Globals.md
@@ -55,131 +55,30 @@
 
 Name | Description
 ---- | -----------
-`build` | [BuildConfig](#buildconfig)<br>Configuration for the build plugin
 `compute` | [ComputeConfig](#computeconfig)<br>Configuration for the compute plugin
 `analytics` | [AnalyticsConfig](#analyticsconfig)<br>Configuration for the analytics plugin
+`build` | [BuildConfig](#buildconfig)<br>Configuration for the build plugin
 `exampleplugin` | [ExamplePluginConfig](#examplepluginconfig)<br>Configuration for the exampleplugin plugin
-`storage` | [StorageConfig](#storageconfig)<br>Configuration for the storage plugin
 `experimental` | [ExperimentalConfig](#experimentalconfig)<br>Configuration for the experimental plugin
 `symbols` | [SymbolsConfig](#symbolsconfig)<br>Configuration for the symbols plugin
 `tools` | [ToolsConfig](#toolsconfig)<br>Configuration for the tools plugin
 `secrets` | [SecretsConfig](#secretsconfig)<br>Configuration for the secrets plugin
+`storage` | [StorageConfig](#storageconfig)<br>Configuration for the storage plugin
 
-## BuildConfig
+## ComputeConfig
 
 Name | Description
 ---- | -----------
-`perforceClusters` | [PerforceCluster](#perforcecluster)`[]`<br>
-`devices` | [DeviceConfig](#deviceconfig)<br>
-`maxConformCount` | `integer`<br>
-`agentShutdownIfDisabledGracePeriod` | `string`<br>
-`artifactTypes` | [ArtifactTypeConfig](#artifacttypeconfig)`[]`<br>
-`projects` | [ProjectConfig](Projects.md)`[]`<br>
-`enableConformTasks` | `boolean`<br>
-`enableJobTaskConformCheck` | `boolean`<br>
-`issueFixedTag` | `string`<br>
-`issues` | [IssueConfig](Issues.md)<br>
-`enableIssueRootCauseAnalysisNotifications` | `boolean`<br>
-`autoRcaUseMultiStage` | `boolean`<br>
-`maxConcurrentRcaRequests` | `integer`<br>
-`issueFindSuspectsMaxSuccessDays` | `integer`<br>
-`issueFindSuspectsMaxScoringSeconds` | `integer`<br>
-`issueFindSuspectsMaxChanges` | `integer`<br>
-`issueFindSuspectsEnableParallelScoring` | `boolean`<br>
-`issueFindSuspectsMaxParallelism` | `integer`<br>
-`issueFindSuspectsEnableEarlyExit` | `boolean`<br>
-`issueFindSuspectsEnableBidirectionalQuery` | `boolean`<br>
-`issueFindSuspectsEnableFingerprintFallback` | `boolean`<br>
-`issueFingerprintFallbackMinScore` | `integer`<br>
-`issueFingerprintFallbackMaxAgeDays` | `integer`<br>
-`issueFingerprintFallbackExcludedTypes` | `string[]`<br>
-`enableExpandedIgnorePatterns` | `boolean`<br>
-
-## PerforceCluster
-
-Name | Description
----- | -----------
-`name` | `string`<br>
-`serviceAccount` | `string`<br>
-`canImpersonate` | `boolean`<br>
-`supportsPartitionedWorkspaces` | `boolean`<br>
-`servers` | [PerforceServer](#perforceserver)`[]`<br>
-`credentials` | [PerforceCredentials](#perforcecredentials)`[]`<br>
-`autoSdk` | [AutoSdkWorkspace](#autosdkworkspace)`[]`<br>
-
-## PerforceServer
-
-Name | Description
----- | -----------
-`serverAndPort` | `string`<br>
-`healthCheck` | `boolean`<br>
-`resolveDns` | `boolean`<br>
-`maxConformCount` | `integer`<br>
-`condition` | `string`<br>
-`properties` | `string[]`<br>
-
-## PerforceCredentials
-
-Name | Description
----- | -----------
-`userName` | `string`<br>
-`password` | `string`<br>
-`ticket` | `string`<br>
-
-## AutoSdkWorkspace
-
-Name | Description
----- | -----------
-`name` | `string`<br>
-`properties` | `string[]`<br>
-`userName` | `string`<br>
-`stream` | `string`<br>
-
-## DeviceConfig
-
-Name | Description
----- | -----------
-`platforms` | [DevicePlatformConfig](#deviceplatformconfig)`[]`<br>
-`pools` | [DevicePoolConfig](#devicepoolconfig)`[]`<br>
-
-## DevicePlatformConfig
-
-Name | Description
----- | -----------
-`id` | `string`<br>
-`name` | `string`<br>
-`models` | `string[]`<br>
-`legacyNames` | `string[]`<br>
-`legacyPerfSpecHighModel` | `string`<br>
-
-## DevicePoolConfig
-
-Name | Description
----- | -----------
-`id` | `string`<br>
-`name` | `string`<br>
-`poolType` | [DevicePoolType](#devicepooltype-enum)<br>
-`projectIds` | `string[]`<br>
-
-## DevicePoolType (Enum)
-
-The type of device pool
-
-Name | Description
----- | -----------
-`Automation` | Available to CIS jobs
-`Shared` | Shared by users with remote checking and checkouts
-
-## ArtifactTypeConfig
-
-Name | Description
----- | -----------
-`name` | `string`<br>
-`type` | `string`<br>
+`include` | [ConfigInclude](#configinclude)`[]`<br>
+`macros` | [ConfigMacro](#configmacro)`[]`<br>
 `acl` | [AclConfig](#aclconfig)<br>
-`keepCount` | `integer`<br>
-`keepDays` | `integer`<br>
-`namespaceId` | `string`<br>
+`versionEnum` | [ConfigVersion](#configversion-enum)<br>
+`telemetryStoreId` | `string`<br>
+`rates` | [AgentRateConfig](#agentrateconfig)`[]`<br>
+`clusters` | [ComputeClusterConfig](#computeclusterconfig)`[]`<br>
+`pools` | [PoolConfig](#poolconfig)`[]`<br>
+`software` | [AgentSoftwareConfig](#agentsoftwareconfig)`[]`<br>
+`networks` | [NetworkConfig](#networkconfig)`[]`<br>
 
 ## AclConfig
 
@@ -214,21 +113,6 @@
 `excludeActions` | `string[]`<br>
 `extends` | `string[]`<br>
 
-## ComputeConfig
-
-Name | Description
----- | -----------
-`include` | [ConfigInclude](#configinclude)`[]`<br>
-`macros` | [ConfigMacro](#configmacro)`[]`<br>
-`acl` | [AclConfig](#aclconfig)<br>
-`versionEnum` | [ConfigVersion](#configversion-enum)<br>
-`telemetryStoreId` | `string`<br>
-`rates` | [AgentRateConfig](#agentrateconfig)`[]`<br>
-`clusters` | [ComputeClusterConfig](#computeclusterconfig)`[]`<br>
-`pools` | [PoolConfig](#poolconfig)`[]`<br>
-`software` | [AgentSoftwareConfig](#agentsoftwareconfig)`[]`<br>
-`networks` | [NetworkConfig](#networkconfig)`[]`<br>
-
 ## ConfigVersion (Enum)
 
 Name | Description
@@ -387,73 +271,127 @@
 ---- | -----------
 `stores` | [TelemetryStoreConfig](Telemetry.md)`[]`<br>
 
+## BuildConfig
+
+Name | Description
+---- | -----------
+`perforceClusters` | [PerforceCluster](#perforcecluster)`[]`<br>
+`devices` | [DeviceConfig](#deviceconfig)<br>
+`maxConformCount` | `integer`<br>
+`agentShutdownIfDisabledGracePeriod` | `string`<br>
+`artifactTypes` | [ArtifactTypeConfig](#artifacttypeconfig)`[]`<br>
+`projects` | [ProjectConfig](Projects.md)`[]`<br>
+`enableConformTasks` | `boolean`<br>
+`enableJobTaskConformCheck` | `boolean`<br>
+`issueFixedTag` | `string`<br>
+`issues` | [IssueConfig](Issues.md)<br>
+`enableIssueRootCauseAnalysisNotifications` | `boolean`<br>
+`autoRcaUseMultiStage` | `boolean`<br>
+`maxConcurrentRcaRequests` | `integer`<br>
+`issueFindSuspectsMaxSuccessDays` | `integer`<br>
+`issueFindSuspectsMaxScoringSeconds` | `integer`<br>
+`issueFindSuspectsMaxChanges` | `integer`<br>
+`issueFindSuspectsEnableParallelScoring` | `boolean`<br>
+`issueFindSuspectsMaxParallelism` | `integer`<br>
+`issueFindSuspectsEnableEarlyExit` | `boolean`<br>
+`issueFindSuspectsEnableBidirectionalQuery` | `boolean`<br>
+`issueFindSuspectsEnableFingerprintFallback` | `boolean`<br>
+`issueFingerprintFallbackMinScore` | `integer`<br>
+`issueFingerprintFallbackMaxAgeDays` | `integer`<br>
+`issueFingerprintFallbackExcludedTypes` | `string[]`<br>
+`enableExpandedIgnorePatterns` | `boolean`<br>
+
+## PerforceCluster
+
+Name | Description
+---- | -----------
+`name` | `string`<br>
+`serviceAccount` | `string`<br>
+`canImpersonate` | `boolean`<br>
+`supportsPartitionedWorkspaces` | `boolean`<br>
+`servers` | [PerforceServer](#perforceserver)`[]`<br>
+`credentials` | [PerforceCredentials](#perforcecredentials)`[]`<br>
+`autoSdk` | [AutoSdkWorkspace](#autosdkworkspace)`[]`<br>
+
+## PerforceServer
+
+Name | Description
+---- | -----------
+`serverAndPort` | `string`<br>
+`healthCheck` | `boolean`<br>
+`resolveDns` | `boolean`<br>
+`maxConformCount` | `integer`<br>
+`condition` | `string`<br>
+`properties` | `string[]`<br>
+
+## PerforceCredentials
+
+Name | Description
+---- | -----------
+`userName` | `string`<br>
+`password` | `string`<br>
+`ticket` | `string`<br>
+
+## AutoSdkWorkspace
+
+Name | Description
+---- | -----------
+`name` | `string`<br>
+`properties` | `string[]`<br>
+`userName` | `string`<br>
+`stream` | `string`<br>
+
+## DeviceConfig
+
+Name | Description
+---- | -----------
+`platforms` | [DevicePlatformConfig](#deviceplatformconfig)`[]`<br>
+`pools` | [DevicePoolConfig](#devicepoolconfig)`[]`<br>
+
+## DevicePlatformConfig
+
+Name | Description
+---- | -----------
+`id` | `string`<br>
+`name` | `string`<br>
+`models` | `string[]`<br>
+`legacyNames` | `string[]`<br>
+`legacyPerfSpecHighModel` | `string`<br>
+
+## DevicePoolConfig
+
+Name | Description
+---- | -----------
+`id` | `string`<br>
+`name` | `string`<br>
+`poolType` | [DevicePoolType](#devicepooltype-enum)<br>
+`projectIds` | `string[]`<br>
+
+## DevicePoolType (Enum)
+
+The type of device pool
+
+Name | Description
+---- | -----------
+`Automation` | Available to CIS jobs
+`Shared` | Shared by users with remote checking and checkouts
+
+## ArtifactTypeConfig
+
+Name | Description
+---- | -----------
+`name` | `string`<br>
+`type` | `string`<br>
+`acl` | [AclConfig](#aclconfig)<br>
+`keepCount` | `integer`<br>
+`keepDays` | `integer`<br>
+`namespaceId` | `string`<br>
+
 ## ExamplePluginConfig
 
 Name | Description
 ---- | -----------
 `greetingMessage` | `string`<br>
 
-## StorageConfig
-
-Name | Description
----- | -----------
-`enableGc` | `boolean`<br>
-`enableGcVerification` | `boolean`<br>
-`backends` | [BackendConfig](#backendconfig)`[]`<br>
-`namespaces` | [NamespaceConfig](#namespaceconfig)`[]`<br>
-
-## BackendConfig
-
-Name | Description
----- | -----------
-`id` | `string`<br>
-`base` | `string`<br>
-`secondary` | `string`<br>
-`type` | [StorageBackendType](#storagebackendtype-enum)<br>
-`baseDir` | `string`<br>
-`awsBucketName` | `string`<br>
-`awsBucketPath` | `string`<br>
-`awsCredentials` | [AwsCredentialsType](#awscredentialstype-enum)<br>
-`awsRole` | `string`<br>
-`awsProfile` | `string`<br>
-`awsRegion` | `string`<br>
-`azureConnectionString` | `string`<br>
-`azureContainerName` | `string`<br>
-`relayServer` | `string`<br>
-`relayToken` | `string`<br>
-`gcsBucketName` | `string`<br>
-`gcsBucketPath` | `string`<br>
-
-## StorageBackendType (Enum)
-
-Name | Description
----- | -----------
-`FileSystem` | 
-`Aws` | 
-`Azure` | 
-`Gcs` | 
-`Memory` | 
-
-## AwsCredentialsType (Enum)
-
-Name | Description
----- | -----------
-`Default` | 
-`Profile` | 
-`AssumeRole` | 
-`AssumeRoleWebIdentity` | 
-
-## NamespaceConfig
-
-Name | Description
----- | -----------
-`id` | `string`<br>
-`backend` | `string`<br>
-`prefix` | `string`<br>
-`gcFrequencyHrs` | `number`<br>
-`gcDelayHrs` | `number`<br>
-`enableAliases` | `boolean`<br>
-`acl` | [AclConfig](#aclconfig)<br>
-
 ## ExperimentalConfig
 
@@ -565,3 +503,65 @@
 Name | Description
 ---- | -----------
 `vaultUri` | `string`<br>
+
+## StorageConfig
+
+Name | Description
+---- | -----------
+`enableGc` | `boolean`<br>
+`enableGcVerification` | `boolean`<br>
+`backends` | [BackendConfig](#backendconfig)`[]`<br>
+`namespaces` | [NamespaceConfig](#namespaceconfig)`[]`<br>
+
+## BackendConfig
+
+Name | Description
+---- | -----------
+`id` | `string`<br>
+`base` | `string`<br>
+`secondary` | `string`<br>
+`type` | [StorageBackendType](#storagebackendtype-enum)<br>
+`baseDir` | `string`<br>
+`awsBucketName` | `string`<br>
+`awsBucketPath` | `string`<br>
+`awsCredentials` | [AwsCredentialsType](#awscredentialstype-enum)<br>
+`awsRole` | `string`<br>
+`awsProfile` | `string`<br>
+`awsRegion` | `string`<br>
+`azureConnectionString` | `string`<br>
+`azureContainerName` | `string`<br>
+`relayServer` | `string`<br>
+`relayToken` | `string`<br>
+`gcsBucketName` | `string`<br>
+`gcsBucketPath` | `string`<br>
+
+## StorageBackendType (Enum)
+
+Name | Description
+---- | -----------
+`FileSystem` | 
+`Aws` | 
+`Azure` | 
+`Gcs` | 
+`Memory` | 
+
+## AwsCredentialsType (Enum)
+
+Name | Description
+---- | -----------
+`Default` | 
+`Profile` | 
+`AssumeRole` | 
+`AssumeRoleWebIdentity` | 
+
+## NamespaceConfig
+
+Name | Description
+---- | -----------
+`id` | `string`<br>
+`backend` | `string`<br>
+`prefix` | `string`<br>
+`gcFrequencyHrs` | `number`<br>
+`gcDelayHrs` | `number`<br>
+`enableAliases` | `boolean`<br>
+`acl` | [AclConfig](#aclconfig)<br>
diff --git a/Engine/Source/Programs/Horde/Docs/Deployment/ServerSettings.md b/Engine/Source/Programs/Horde/Docs/Deployment/ServerSettings.md
--- a/Engine/Source/Programs/Horde/Docs/Deployment/ServerSettings.md
+++ b/Engine/Source/Programs/Horde/Docs/Deployment/ServerSettings.md
@@ -125,15 +125,86 @@
 
 Name | Description
 ---- | -----------
-`build` | [BuildServerConfig](#buildserverconfig)<br>Configuration for the build plugin
 `compute` | [ComputeServerConfig](#computeserverconfig)<br>Configuration for the compute plugin
 `analytics` | [AnalyticsServerConfig](#analyticsserverconfig)<br>Configuration for the analytics plugin
+`build` | [BuildServerConfig](#buildserverconfig)<br>Configuration for the build plugin
 `exampleplugin` | [ExamplePluginServerConfig](#examplepluginserverconfig)<br>Configuration for the exampleplugin plugin
-`storage` | [StorageServerConfig](#storageserverconfig)<br>Configuration for the storage plugin
 `experimental` | [ExperimentalServerConfig](#experimentalserverconfig)<br>Configuration for the experimental plugin
 `symbols` | [PluginServerConfig](#pluginserverconfig)<br>Configuration for the symbols plugin
 `tools` | [ToolsServerConfig](#toolsserverconfig)<br>Configuration for the tools plugin
 `secrets` | [SecretsServerConfig](#secretsserverconfig)<br>Configuration for the secrets plugin
+`storage` | [StorageServerConfig](#storageserverconfig)<br>Configuration for the storage plugin
 
+## ComputeServerConfig
+
+Name | Description
+---- | -----------
+`enableUpgradeTasks` | `boolean`<br>
+`withAws` | `boolean`<br>
+`awsRegions` | `string[]`<br>
+`awsAutoScalingQueueUrls` | `string[]`<br>
+`fleetManagerV2` | [FleetManagerType](#fleetmanagertype-enum)<br>
+`fleetManagerV2Config` | `object`<br>
+`autoEnrollAgents` | `boolean`<br>
+`defaultAgentPoolSizeStrategy` | [PoolSizeStrategy](#poolsizestrategy-enum)<br>
+`agentPoolScaleOutCooldownSeconds` | `integer`<br>
+`agentPoolScaleInCooldownSeconds` | `integer`<br>
+`sessionConflictEnabled` | `boolean`<br>
+`sessionConflictThreshold` | `integer`<br>
+`sessionConflictWindowSeconds` | `integer`<br>
+`computeTunnelPort` | `integer`<br>
+`computeTunnelAddress` | `string`<br>
+`enabled` | `boolean`<br>
+
+## FleetManagerType (Enum)
+
+Name | Description
+---- | -----------
+`Default` | 
+`NoOp` | 
+`Aws` | 
+`AwsReuse` | 
+`AwsRecycle` | 
+`AwsAsg` | 
+
+## PoolSizeStrategy (Enum)
+
+Name | Description
+---- | -----------
+`LeaseUtilization` | 
+`JobQueue` | 
+`NoOp` | 
+`ComputeQueueAwsMetric` | 
+`LeaseUtilizationAwsMetric` | 
+
+## AnalyticsServerConfig
+
+Name | Description
+---- | -----------
+`sinks` | [TelemetrySinkConfig](#telemetrysinkconfig)<br>
+`enabled` | `boolean`<br>
+
+## TelemetrySinkConfig
+
+Name | Description
+---- | -----------
+`epic` | [EpicTelemetryConfig](#epictelemetryconfig)<br>
+`mongo` | [MongoTelemetryConfig](#mongotelemetryconfig)<br>
+
+## EpicTelemetryConfig
+
+Name | Description
+---- | -----------
+`url` | `string`<br>
+`appId` | `string`<br>
+`enabled` | `boolean`<br>
+
+## MongoTelemetryConfig
+
+Name | Description
+---- | -----------
+`retainDays` | `number`<br>
+`enabled` | `boolean`<br>
+
 ## BuildServerConfig
 
@@ -271,131 +342,12 @@
 `maxChildCount` | `integer`<br>Maximum number of children in each node
 `sliceThreshold` | `integer`<br>Threshold hash value for splitting interior nodes
 
-## ComputeServerConfig
-
-Name | Description
----- | -----------
-`enableUpgradeTasks` | `boolean`<br>
-`withAws` | `boolean`<br>
-`awsRegions` | `string[]`<br>
-`awsAutoScalingQueueUrls` | `string[]`<br>
-`fleetManagerV2` | [FleetManagerType](#fleetmanagertype-enum)<br>
-`fleetManagerV2Config` | `object`<br>
-`autoEnrollAgents` | `boolean`<br>
-`defaultAgentPoolSizeStrategy` | [PoolSizeStrategy](#poolsizestrategy-enum)<br>
-`agentPoolScaleOutCooldownSeconds` | `integer`<br>
-`agentPoolScaleInCooldownSeconds` | `integer`<br>
-`computeTunnelPort` | `integer`<br>
-`computeTunnelAddress` | `string`<br>
-`enabled` | `boolean`<br>
-
-## FleetManagerType (Enum)
-
-Name | Description
----- | -----------
-`Default` | 
-`NoOp` | 
-`Aws` | 
-`AwsReuse` | 
-`AwsRecycle` | 
-`AwsAsg` | 
-
-## PoolSizeStrategy (Enum)
-
-Name | Description
----- | -----------
-`LeaseUtilization` | 
-`JobQueue` | 
-`NoOp` | 
-`ComputeQueueAwsMetric` | 
-`LeaseUtilizationAwsMetric` | 
-
-## AnalyticsServerConfig
-
-Name | Description
----- | -----------
-`sinks` | [TelemetrySinkConfig](#telemetrysinkconfig)<br>
-`enabled` | `boolean`<br>
-
-## TelemetrySinkConfig
-
-Name | Description
----- | -----------
-`epic` | [EpicTelemetryConfig](#epictelemetryconfig)<br>
-`mongo` | [MongoTelemetryConfig](#mongotelemetryconfig)<br>
-
-## EpicTelemetryConfig
-
-Name | Description
----- | -----------
-`url` | `string`<br>
-`appId` | `string`<br>
-`enabled` | `boolean`<br>
-
-## MongoTelemetryConfig
-
-Name | Description
----- | -----------
-`retainDays` | `number`<br>
-`enabled` | `boolean`<br>
-
 ## ExamplePluginServerConfig
 
 Name | Description
 ---- | -----------
 `enabled` | `boolean`<br>
 
-## StorageServerConfig
-
-Name | Description
----- | -----------
-`bundleCacheDir` | `string`<br>
-`bundleCacheSize` | `string`<br>
-`bundleCacheSizeBytes` | `integer`<br>
-`backends` | [BackendConfig](#backendconfig)`[]`<br>
-`enabled` | `boolean`<br>
-
-## BackendConfig
-
-Name | Description
----- | -----------
-`id` | `string`<br>
-`base` | `string`<br>
-`secondary` | `string`<br>
-`type` | [StorageBackendType](#storagebackendtype-enum)<br>
-`baseDir` | `string`<br>
-`awsBucketName` | `string`<br>
-`awsBucketPath` | `string`<br>
-`awsCredentials` | [AwsCredentialsType](#awscredentialstype-enum)<br>
-`awsRole` | `string`<br>
-`awsProfile` | `string`<br>
-`awsRegion` | `string`<br>
-`azureConnectionString` | `string`<br>
-`azureContainerName` | `string`<br>
-`relayServer` | `string`<br>
-`relayToken` | `string`<br>
-`gcsBucketName` | `string`<br>
-`gcsBucketPath` | `string`<br>
-
-## StorageBackendType (Enum)
-
-Name | Description
----- | -----------
-`FileSystem` | 
-`Aws` | 
-`Azure` | 
-`Gcs` | 
-`Memory` | 
-
-## AwsCredentialsType (Enum)
-
-Name | Description
----- | -----------
-`Default` | 
-`Profile` | 
-`AssumeRole` | 
-`AssumeRoleWebIdentity` | 
-
 ## ExperimentalServerConfig
 
 Name | Description
@@ -475,3 +427,54 @@
 ---- | -----------
 `withAws` | `boolean`<br>
 `enabled` | `boolean`<br>
+
+## StorageServerConfig
+
+Name | Description
+---- | -----------
+`bundleCacheDir` | `string`<br>
+`bundleCacheSize` | `string`<br>
+`bundleCacheSizeBytes` | `integer`<br>
+`backends` | [BackendConfig](#backendconfig)`[]`<br>
+`enabled` | `boolean`<br>
+
+## BackendConfig
+
+Name | Description
+---- | -----------
+`id` | `string`<br>
+`base` | `string`<br>
+`secondary` | `string`<br>
+`type` | [StorageBackendType](#storagebackendtype-enum)<br>
+`baseDir` | `string`<br>
+`awsBucketName` | `string`<br>
+`awsBucketPath` | `string`<br>
+`awsCredentials` | [AwsCredentialsType](#awscredentialstype-enum)<br>
+`awsRole` | `string`<br>
+`awsProfile` | `string`<br>
+`awsRegion` | `string`<br>
+`azureConnectionString` | `string`<br>
+`azureContainerName` | `string`<br>
+`relayServer` | `string`<br>
+`relayToken` | `string`<br>
+`gcsBucketName` | `string`<br>
+`gcsBucketPath` | `string`<br>
+
+## StorageBackendType (Enum)
+
+Name | Description
+---- | -----------
+`FileSystem` | 
+`Aws` | 
+`Azure` | 
+`Gcs` | 
+`Memory` | 
+
+## AwsCredentialsType (Enum)
+
+Name | Description
+---- | -----------
+`Default` | 
+`Profile` | 
+`AssumeRole` | 
+`AssumeRoleWebIdentity` | 
