Horde: Make agent session conflict detection disabled as default

diff --git a/Engine/Source/Programs/Horde/Plugins/Build/HordeServer.Build.Tests/Server/RpcServiceTest.cs b/Engine/Source/Programs/Horde/Plugins/Build/HordeServer.Build.Tests/Server/RpcServiceTest.cs
--- a/Engine/Source/Programs/Horde/Plugins/Build/HordeServer.Build.Tests/Server/RpcServiceTest.cs
+++ b/Engine/Source/Programs/Horde/Plugins/Build/HordeServer.Build.Tests/Server/RpcServiceTest.cs
@@ -252,7 +252,9 @@ namespace HordeServer.Tests.Server
 		[TestMethod]
 		public async Task UpdateSession_SessionMismatch_QuarantinesAfterThresholdAsync()
 		{
-			int threshold = ServiceProvider.GetRequiredService<IOptionsSnapshot<ComputeServerConfig>>().Value.SessionConflictThreshold;
+			ComputeServerConfig computeServerConfig = ServiceProvider.GetRequiredService<IOptionsMonitor<ComputeServerConfig>>().CurrentValue;
+			computeServerConfig.SessionConflictEnabled = true;
+			int threshold = computeServerConfig.SessionConflictThreshold;
 			AgentId agentId = new AgentId("conflict-agent");
 			IPAddress callerIp = IPAddress.Parse("10.0.0.1");
 			ServerCallContext adminCtx = ServerCallContextStub.ForAdmin(callerIp);
diff --git a/Engine/Source/Programs/Horde/Plugins/Compute/HordeServer.Compute.Tests/Agents/SessionConflictServiceTest.cs b/Engine/Source/Programs/Horde/Plugins/Compute/HordeServer.Compute.Tests/Agents/SessionConflictServiceTest.cs
--- a/Engine/Source/Programs/Horde/Plugins/Compute/HordeServer.Compute.Tests/Agents/SessionConflictServiceTest.cs
+++ b/Engine/Source/Programs/Horde/Plugins/Compute/HordeServer.Compute.Tests/Agents/SessionConflictServiceTest.cs
@@ -37,7 +37,7 @@ namespace HordeServer.Tests.Agents
 		[TestMethod]
 		public async Task BelowThreshold_NoQuarantineAsync()
 		{
-			SessionConflictService service = CreateService();
+			SessionConflictService service = CreateService(new ComputeServerConfig { SessionConflictEnabled = true });
 			IAgent agent = await AgentCollection.AddAsync(new CreateAgentOptions(new AgentId("conflict-test-1"), AgentMode.Dedicated, false, ""));
 
 			// Record 2 conflicts (threshold is 3)
@@ -52,7 +52,7 @@ namespace HordeServer.Tests.Agents
 		[TestMethod]
 		public async Task AtThreshold_QuarantinesAgentAsync()
 		{
-			SessionConflictService service = CreateService();
+			SessionConflictService service = CreateService(new ComputeServerConfig { SessionConflictEnabled = true });
 			IAgent agent = await AgentCollection.AddAsync(new CreateAgentOptions(new AgentId("conflict-test-2"), AgentMode.Dedicated, false, ""));
 
 			// Record 3 conflicts (hits threshold)
@@ -72,7 +72,7 @@ namespace HordeServer.Tests.Agents
 		[TestMethod]
 		public async Task WindowExpiry_OldEventsNotCountedAsync()
 		{
-			ComputeServerConfig config = new ComputeServerConfig { SessionConflictWindowSeconds = 300 };
+			ComputeServerConfig config = new ComputeServerConfig { SessionConflictEnabled = true, SessionConflictWindowSeconds = 300 };
 			SessionConflictService service = CreateService(config);
 			IAgent agent = await AgentCollection.AddAsync(new CreateAgentOptions(new AgentId("conflict-test-3"), AgentMode.Dedicated, false, ""));
 
@@ -98,7 +98,7 @@ namespace HordeServer.Tests.Agents
 		[TestMethod]
 		public async Task AlreadyDisabled_IdempotentAsync()
 		{
-			SessionConflictService service = CreateService();
+			SessionConflictService service = CreateService(new ComputeServerConfig { SessionConflictEnabled = true });
 			IAgent agent = await AgentCollection.AddAsync(new CreateAgentOptions(new AgentId("conflict-test-4"), AgentMode.Dedicated, false, ""));
 
 			// Pre-disable the agent with a custom comment
diff --git a/Engine/Source/Programs/Horde/Plugins/Compute/HordeServer.Compute/ComputeServerConfig.cs b/Engine/Source/Programs/Horde/Plugins/Compute/HordeServer.Compute/ComputeServerConfig.cs
--- a/Engine/Source/Programs/Horde/Plugins/Compute/HordeServer.Compute/ComputeServerConfig.cs
+++ b/Engine/Source/Programs/Horde/Plugins/Compute/HordeServer.Compute/ComputeServerConfig.cs
@@ -72,7 +72,7 @@ namespace HordeServer
 		/// When multiple machines accidentally connect using the same agent ID (e.g. cloned VMs),
 		/// repeated session mismatches will cause the agent to be automatically disabled.
 		/// </summary>
-		public bool SessionConflictEnabled { get; set; } = true;
+		public bool SessionConflictEnabled { get; set; } = false;
 
 		/// <summary>
 		/// Number of session mismatches within the sliding window required to trigger auto-quarantine
