Horde: Dashboard Release Notes

diff --git a/Engine/Source/Programs/Horde/Docs/ReleaseNotes.md b/Engine/Source/Programs/Horde/Docs/ReleaseNotes.md
--- a/Engine/Source/Programs/Horde/Docs/ReleaseNotes.md
+++ b/Engine/Source/Programs/Horde/Docs/ReleaseNotes.md
@@ -1,5 +1,63 @@
 # Release Notes
 
+## 2026-02-26
+
+* Make agent session conflict detection disabled as default (51262410)
+* Invite users nominated to fix build health issue to Slack triage channel
+  Add InviteUsersAsync call after posting nomination message in CreateOrUpdateWorkflowThreadAsync, ensuring the user receives a notification.
+  Suppress already_in_channel warning in non-admin invite path, matching existing admin path behavior, to reduce noise in logs. (51241288)
+* Add session conflict detection and auto-quarantine to agents
+  Detect agents experiencing repeated session ID mismatches (e.g. cloned VMs sharing an agent ID) and automatically quarantine them after a configurable threshold. Conflicts are tracked in Redis using a sliding-window sorted set; once the threshold is reached the agent is disabled with a diagnostic comment listing the offending IPs. Suppresses log spam for already-quarantined agents and surfaces them in the daily Slack agent report. (51232325)
+* Make log and HTTP blob uploading more robust
+  - Use a bare SocketsHttpHandler for upload redirect clients instead of the Horde-server-tuned Polly retry/timeout policies
+  - Prevent ServerLogger.DisposeAsync from propagating exceptions thrown by the tail task or blob writer
+  - Fix redirect uploads to presigned URLs sending an empty body on retry by creating a fresh StreamContent and resetting stream position per attempt (51224648)
+* Change conform check when assigning a job task to have up to date free disk space for the agent
+  On lease completion disk free space property of the agent is updated to become available for the next lease assignment.
+  Implements a minimal ISystemMetrics for macOS to support disk usage. (51212055)
+* Limit size of log messages in server
+  Too large log messages can cause issues with external log handling pipelines (containerd/FluentBit/OTel etc). This change prevents too deep destructuring of exceptions, a common source of over-sized messages.
+  Also inserts a JSON formatter for Serilog that caps log entries by size, replacing them with a compact JSON stub with the bare minimum information. (51211255)
+* Optimize FindIssuesV2 to use aggregate for distinct issue IDs and fetch spans only for paginated results (51210130)
+* Tune the step events for job step history (51208415)
+* Update TestIgnorePattern_StructuredProperties for LogValue
+  Test now validates IgnoredLine as a LogValue with nested properties
+  instead of top-level state entries. (51183653)
+* Fix breadcrumbs on tools view (51172192)
+* Add batch log events API and inline error/warning display in step history
+  New POST /api/v1/logs/events endpoint for fetching events across multiple logs
+  in a single request with severity filtering and count limiting. Step history
+  panel shows expandable error/warning events inline with log links. (51164657)
+* Add time-based persistent flushing to ServerLogger
+  ServerLogger previously only flushed buffered log data to blob storage when the buffer exceeded 1 MB or at end-of-stream, leaving slow-streaming jobs without visible log updates for several minutes when the tail mechanism was inactive. (51159252)
+* Fix preflight commit ID parsing (51158668)
+* Fix up and unify change list preflight display (51157338)
+* Add batch log events API and step history error/warning display
+  Add POST /api/v1/logs/events endpoint for fetching events across multiple logs in a single request.
+  Display error/warning events inline in step history table with expand/collapse buttons. (51157086)
+* Ensure perforce service check (51155190)
+* Fix preflight config check warning overflow (51154487)
+* Fix compilation warnings (51100244)
+* Broaden exception handling in StreamsController batch commit retrieval
+  Catch all exceptions (not just PerforceException) when retrieving batch commits so failures are logged and reported via MissingCommitsHeader regardless of the error type. (51096190)
+* Add include support to ComputeConfig for pool configuration
+  Allows pool definitions to be split across multiple included files
+  via the standard ConfigInclude mechanism. (51069581)
+* Resolve diamond include dependencies in config system
+  Detect and deduplicate diamond includes in the config preprocessor.
+  When file A includes B and C and both B and C include D, D is now
+  merged only once per [ConfigIncludeRoot] scope. A warning is emitted
+  identifying both the first includer and the duplicate requester.
+  Add ResolveWithOverridesAsync to IConfigService so callers can get
+  the fully resolved GlobalConfig (with plugins) from a set of file
+  overrides. Used by the pool migration diff endpoint to compare
+  proposed pool configs against the current database.
+  Propagate config warnings through preflight validation, pool-diff,
+  and their corresponding dashboard UIs. (51064451)
+* Add IgnorePatterns documentation (51054716)
+* Batch fetch agent sessions from Redis when loading multiple agents
+  Replace sequential per-agent Redis GET calls in AgentCollection with a single MGET, which reduces Redis round-trips. (51051204)
+
 ## 2026-02-19
 
 * Handle client cancellation gracefully in DownloadSoftware gRPC method (51012208)
