Config loading fault tolerance: missing base references and diamond include scoping

1) MergeDefaults: when a template, agent type, or workspace type references a
base that does not exist, log a warning and skip the merge instead of throwing
an exception that prevents the entire config from loading.

2) Diamond include scoping: the depot code clears ProcessedIncludes only at
IncludeRootDepth==0, so sibling [ConfigIncludeRoot] entries (e.g. two
StreamConfigs in the same project) share the same dedup set. Stream2's
include of base.json gets incorrectly skipped as a diamond when stream1
already included it. Fix: save/restore ProcessedIncludes when entering a
[ConfigIncludeRoot] from property/array traversal (fresh scope per sibling).
When entering from include processing, share the parent scope so within-tree
diamonds are still detected.

diff --git a/Engine/Source/Programs/Horde/HordeServer.Shared/Configuration/ConfigContext.cs b/Engine/Source/Programs/Horde/HordeServer.Shared/Configuration/ConfigContext.cs
--- a/Engine/Source/Programs/Horde/HordeServer.Shared/Configuration/ConfigContext.cs
+++ b/Engine/Source/Programs/Horde/HordeServer.Shared/Configuration/ConfigContext.cs
@@ -55,18 +55,25 @@ namespace HordeServer.Configuration
 		/// Tracks files already processed as includes within the current include root scope,
 		/// preventing diamond dependencies from causing the same file to be merged multiple times.
 		/// Maps each included URI to the URI of the file that first included it.
-		/// Cleared when entering each new [ConfigIncludeRoot] hierarchy so that different
-		/// include roots (e.g. GlobalConfig vs ProjectConfig) can independently include the same file.
+		/// Each [ConfigIncludeRoot] entered from property/array traversal saves and restores this
+		/// so sibling roots (e.g. two StreamConfigs in the same project) get independent scopes.
+		/// Roots entered from include processing share the parent's scope for diamond detection.
 		/// </summary>
-		public Dictionary<Uri, Uri?> ProcessedIncludes { get; } = new Dictionary<Uri, Uri?>();
+		public Dictionary<Uri, Uri?> ProcessedIncludes { get; set; } = new Dictionary<Uri, Uri?>();
 
 		/// <summary>
-		/// Depth counter for nested [ConfigIncludeRoot] processing. When zero, the next
-		/// IncludeRoot entry is a new top-level scope that should clear ProcessedIncludes.
-		/// When positive, we're recursing within the same include tree (shared dedup set).
+		/// Depth counter for nested [ConfigIncludeRoot] processing.
 		/// </summary>
 		public int IncludeRootDepth { get; set; }
 
+		/// <summary>
+		/// When true, the current PreprocessAsync call is processing an included file (recursive
+		/// from the include loop). [ConfigIncludeRoot] entries should share the parent's diamond
+		/// dedup scope rather than creating a fresh one, so that diamonds within the same include
+		/// tree are still detected.
+		/// </summary>
+		public bool InIncludeScope { get; set; }
+
 		/// <summary>
 		/// Warnings generated during config processing (e.g. diamond include dependencies)
 		/// </summary>
diff --git a/Engine/Source/Programs/Horde/HordeServer.Shared/Configuration/ConfigNode.cs b/Engine/Source/Programs/Horde/HordeServer.Shared/Configuration/ConfigNode.cs
--- a/Engine/Source/Programs/Horde/HordeServer.Shared/Configuration/ConfigNode.cs
+++ b/Engine/Source/Programs/Horde/HordeServer.Shared/Configuration/ConfigNode.cs
@@ -538,15 +538,18 @@ namespace HordeServer.Configuration
 			// Before parsing properties into this object, read all the includes recursively
 			if (IncludeRoot)
 			{
-				// Each [ConfigIncludeRoot] type gets its own scope for diamond include dedup.
-				// When IncludeRootDepth is 0, we're entering a new top-level root (e.g. GlobalConfig,
-				// ProjectConfig, StreamConfig) and should clear the dedup set. When positive, we're
-				// recursing within the same include tree and share the set.
-				bool isNewScope = (context.IncludeRootDepth == 0);
-				if (isNewScope)
+				// When entering from property/array processing (not from an include), create a
+				// fresh diamond dedup scope so sibling roots (e.g. two StreamConfigs in the same
+				// project) can independently include the same file.
+				// When entering from include processing (InIncludeScope is true), share the
+				// parent's scope so diamonds within the same include tree are still detected.
+				Dictionary<Uri, Uri?>? savedProcessedIncludes = null;
+				if (!context.InIncludeScope)
 				{
-					context.ProcessedIncludes.Clear();
+					savedProcessedIncludes = context.ProcessedIncludes;
+					context.ProcessedIncludes = new Dictionary<Uri, Uri?>();
 				}
+				bool savedInIncludeScope = context.InIncludeScope;
 				context.IncludeRootDepth++;
 
 				try
@@ -559,6 +562,7 @@ namespace HordeServer.Configuration
 					// processed to handle diamond dependencies (e.g. A includes B and C, both B and C
 					// include D - D should only be merged once).
 					string currentFile = context.CurrentFile?.ToString() ?? "(unknown)";
+					context.InIncludeScope = true;
 					foreach (IConfigFile include in includes)
 					{
 						if (context.ProcessedIncludes.TryGetValue(include.Uri, out Uri? firstIncluder))
@@ -578,7 +582,12 @@ namespace HordeServer.Configuration
 				}
 				finally
 				{
+					context.InIncludeScope = savedInIncludeScope;
 					context.IncludeRootDepth--;
+					if (savedProcessedIncludes != null)
+					{
+						context.ProcessedIncludes = savedProcessedIncludes;
+					}
 				}
 			}
 
diff --git a/Engine/Source/Programs/Horde/HordeServer.Shared/Configuration/ConfigObject.cs b/Engine/Source/Programs/Horde/HordeServer.Shared/Configuration/ConfigObject.cs
--- a/Engine/Source/Programs/Horde/HordeServer.Shared/Configuration/ConfigObject.cs
+++ b/Engine/Source/Programs/Horde/HordeServer.Shared/Configuration/ConfigObject.cs
@@ -3,6 +3,7 @@
 using System.Collections;
 using System.Reflection;
 using EpicGames.Core;
+using Microsoft.Extensions.Logging;
 
 namespace HordeServer.Configuration
 {
@@ -53,7 +54,8 @@ namespace HordeServer.Configuration
 		/// Merges any unassigned properties from the source object to the target
 		/// </summary>
 		/// <param name="objects">Set of objects to use for merging</param>
-		public static void MergeDefaults<TKey, TValue>(IEnumerable<(TKey Key, TKey? BaseKey, TValue Value)> objects)
+		/// <param name="logger">Optional logger. When provided, invalid bases and circular dependencies are logged as warnings instead of throwing exceptions.</param>
+		public static void MergeDefaults<TKey, TValue>(IEnumerable<(TKey Key, TKey? BaseKey, TValue Value)> objects, ILogger? logger = null)
 			where TKey : notnull
 			where TValue : class
 		{
@@ -100,18 +102,46 @@ namespace HordeServer.Configuration
 				{
 					if (!validKeys.Contains(baseKey!))
 					{
-						throw new Exception($"{key} has invalid/missing base {baseKey}");
+						if (logger != null)
+						{
+							logger.LogWarning("{Key} has invalid/missing base {BaseKey}; skipping merge", key, baseKey);
+							handledValues.Add(key, value);
+						}
+						else
+						{
+							throw new Exception($"{key} has invalid/missing base {baseKey}");
+						}
 					}
 				}
 
-				List<TKey> circularKeys = objects.Where(x => !handledValues.ContainsKey(x.Key)).Select(x => x.Key).ToList();
-				if (circularKeys.Count == 1)
-				{
-					throw new Exception($"{circularKeys[0]} has a circular dependency on itself");
-				}
-				else
+				// Remove items that were handled above (invalid base, logged as warning)
+				remainingObjects.RemoveAll(x => handledValues.ContainsKey(x.Key));
+
+				if (remainingObjects.Count > 0)
 				{
-					throw new Exception($"{StringUtils.FormatList(circularKeys.Select(x => x.ToString() ?? "(unknown)"))} have circular dependencies.");
+					List<TKey> circularKeys = remainingObjects.Select(x => x.Key).ToList();
+					if (logger != null)
+					{
+						if (circularKeys.Count == 1)
+						{
+							logger.LogWarning("{Key} has a circular dependency on itself; skipping merge", circularKeys[0]);
+						}
+						else
+						{
+							logger.LogWarning("{Keys} have circular dependencies; skipping merge", StringUtils.FormatList(circularKeys.Select(x => x.ToString() ?? "(unknown)")));
+						}
+					}
+					else
+					{
+						if (circularKeys.Count == 1)
+						{
+							throw new Exception($"{circularKeys[0]} has a circular dependency on itself");
+						}
+						else
+						{
+							throw new Exception($"{StringUtils.FormatList(circularKeys.Select(x => x.ToString() ?? "(unknown)"))} have circular dependencies.");
+						}
+					}
 				}
 			}
 		}
diff --git a/Engine/Source/Programs/Horde/Plugins/Build/HordeServer.Build/Streams/StreamConfig.cs b/Engine/Source/Programs/Horde/Plugins/Build/HordeServer.Build/Streams/StreamConfig.cs
--- a/Engine/Source/Programs/Horde/Plugins/Build/HordeServer.Build/Streams/StreamConfig.cs
+++ b/Engine/Source/Programs/Horde/Plugins/Build/HordeServer.Build/Streams/StreamConfig.cs
@@ -368,9 +368,9 @@ namespace HordeServer.Streams
 				}
 			}
 
-			ConfigObject.MergeDefaults(AgentTypes.Select(x => (x.Key, x.Value.Base, x.Value)));
-			ConfigObject.MergeDefaults(WorkspaceTypes.Select(x => (x.Key, x.Value.Base, x.Value)));
-			ConfigObject.MergeDefaults(Templates.Select(x => (x.Id, x.Base, x)));
+			ConfigObject.MergeDefaults(AgentTypes.Select(x => (x.Key, x.Value.Base, x.Value)), logger);
+			ConfigObject.MergeDefaults(WorkspaceTypes.Select(x => (x.Key, x.Value.Base, x.Value)), logger);
+			ConfigObject.MergeDefaults(Templates.Select(x => (x.Id, x.Base, x)), logger);
 
 			foreach (TemplateRefConfig template in Templates)
 			{
