Horde: Update Build Health documentation with comprehensive workflow reference and examples

diff --git a/Engine/Source/Programs/Horde/Docs/Config/BuildHealth.md b/Engine/Source/Programs/Horde/Docs/Config/BuildHealth.md
--- a/Engine/Source/Programs/Horde/Docs/Config/BuildHealth.md
+++ b/Engine/Source/Programs/Horde/Docs/Config/BuildHealth.md
@@ -2,283 +2,474 @@
 
 # Build Health
 
-Epic has a large development team working on products that are continuously shipping. We
-use a trunk-based development model with branches for each release, with
-commits to release branches automatically being merged back to the mainline.
+Horde's Build Health system is an automated issue tracking and triage pipeline for
+build breakages. It ingests structured log events from build steps, runs
+**issue handlers** to generate **fingerprints**, groups those fingerprints into
+**issues**, identifies **suspect commits**, and notifies developers via **Slack**.
 
-Maintaining a stable and healthy codebase is important to our ability to iterate quickly,
-so we have a number of policies in place to help:
+It works best with trunk-based development where commits to release branches are
+merged back to mainline, and scales to large teams where manual triage is slow.
 
-* Local test suites that developers can run before changes are submitted.
-* The ability to run a suite of pre-submit tests on the build farm, through Horde's _preflight_ system.
-* An automated __Build Health__ system for identifying, grouping, and triaging build errors,
-  and notifying developers responsible for them as quickly as possible via _Slack_.
+## Anatomy of a Build Breakage
 
-The Build Health system can be thought of as similar to a traditional bug tracking database,
-creating __issues__ for build breakages, but is automated, context-aware and highly data
-driven - and tightly integrated into Horde's job dashboard and logs.
+Each branch is modeled as a linear sequence of commits. When a build failure occurs,
+Horde identifies which commit since the last successful build is responsible using
+heuristics such as:
 
-This system is only as good as the data fed into it, so we are careful to annotate any
-errors with as much context as possible to allow the system to make smart decisions when
-things go wrong.
+* Matching modified source files or assets to compile/cook errors.
+* Matching linker symbol names to modified source file names.
+* Tracing merge history to find the originating commit across branches.
 
-## Anatomy of a build breakage
+A single change can cause widespread errors — different compiler diagnostics on
+MSVC vs Clang vs GCC, cook failures across platforms, etc. Build Health groups
+these together and gets eyes on them quickly while avoiding duplicate notifications.
 
-Each branch is modeled as a linear sequence of commits. When a build failure occurs, we
-seek to identify which commit since the last successful build is responsible.
+## Marking Issues as Fixed
 
-There are many heuristics we can do this; looking for commits that modify a particular
-asset or source file that fails to compile, distinguishing between code/content errors and
-changes, looking for linker symbol names that match modified source file names, and so on.
+Issues are marked as fixed automatically after a successful build, but users can
+also indicate a fix via:
 
-Additionally, we may have a single change that causes a widespread set of errors - a source
-file that fails to compile may produce completely different errors on MSVC vs Clang vs GCC,
-for example, a missing asset may cause cook errors for different platforms, and so on. These
-errors may occur at times during the build pipeline, and we want to get eyes on them as soon
-as possible - while also avoiding a long tail of spurious notifications that occur later.
+* **Slack** — interactive buttons on the issue thread.
+* **Horde Dashboard** — the issue detail view.
+* **Commit tags** — including `#horde 1234` on a separate line in a commit description.
 
-Finally, because we develop across multiple branches and merge frequently between them,
-we want to collate errors in different streams and trace back to the original commit that
-is deemed responsible.
+The tag name can be customized via `IssueFixedTag` in
+[globals.json](../Config/Schema/Globals.md#buildconfig).
 
-## Marking Issues as Fixed
+## Issues, Spans, and Fingerprints
 
-Issues will be marked as fixed automatically after a successful build, but users
-can indicate that an issue has been addressed via Slack, the Horde Dashboard, or
-by including a `#horde 1234` tag in a separate line in a commit description.
+Horde creates a build health **issue** for each detected breakage.
 
-The tag name can be customized using the `IssueFixedTag` setting in the
-[globals.json](../Config/Schema/Globals.md#buildconfig) file.
+Each issue contains multiple **spans** — one per stream the issue is observed in —
+identifying the last successful and first failing commit. A span resolves once the
+error no longer occurs; the issue resolves once all spans resolve.
 
-## Issues, Spans, and Fingerprints
+Errors are attached to spans via their **fingerprint**, which describes the error
+programmatically and controls how it matches with other errors.
+
+### Contents of a Fingerprint
 
-Horde creates a build health __issue__ for each detected breakage.
+See `IIssueFingerprint.cs` for the source definition.
 
-Each issue contains multiple __spans__ - one for each stream the issue has been observed in -
-each of which identifies the last successful and first failing commit. A span is marked as
-resolved once the associated error no longer occurs in a stream, and the issue is marked as
-resolved once all its spans are resolved.
+| Field | Description |
+|-------|-------------|
+| **Type** | Identifies the issue type. Fingerprints only match others of the same type (e.g. `Compile`, `Content`, `Linker`). |
+| **SummaryTemplate** | Populates the issue description. May contain [placeholders](#summary-templates). |
+| **Keys** | Identifying data (e.g. filenames) used to match and group fingerprints. |
+| **RejectKeys** | Keys this fingerprint should _not_ match with. |
+| **Metadata** | Arbitrary key/value pairs for summary templates and debugging. |
+| **ChangeFilter** | Semicolon-delimited file wildcards for suspect identification (e.g. `*.c;*.cpp;*.h`). |
 
-Build errors are ingested and attached to spans via their __fingerprint__. Fingerprints
-provide a programmatic description of the error, and how it should be matched with other
-errors.
+### Matching and Merging Fingerprints
 
-### Contents of a Fingerprint
+1. A fingerprint is generated for each error or warning in a build step.
+2. Each is compared against unresolved spans in the same stream. A match requires:
+   - Identical type and summary template.
+   - At least one common key (or both have no keys).
+   - No keys in the other's reject set.
+3. Unmatched fingerprints are grouped using relaxed conditions (identical type,
+   no conflicting reject keys).
+4. For new spans, suspect commits are found via merge history. If an existing issue
+   shares a fingerprint and suspect commit, the span joins that issue; otherwise a
+   new issue is created.
+
+### Creating Fingerprints
+
+Fingerprints can be created by:
+
+* Including fingerprint data directly in a
+  [structured log event](../Internals/StructuredLogging.md) (preferred — most control).
+* Post-processing structured log events on the Horde server after a step completes.
 
-A fingerprint contains the following data (see `IIssueFingerprint.cs`):
+#### Adding Fingerprints to Errors (C#)
 
-* __Type__: A string identifying the type of issue. Fingerprints only match with other fingerprints of the same type.
-  One of the most commonly encountered issue types is `Compile`, indicating a compile error in one or more source files.
-* __SummaryTemplate__: A string that is used to populate an issue's description in the Horde dashboard and notifications.
-  May contain placeholder strings that will be expanded when fingerprints are merged together
-  (see [Summary Templates](#summary-templates)).
-* __Keys__: Pieces of identifying data used to match and group fingerprints,
-  such as the name of a file associated with the error. Each key has a name and type.
-* __RejectKeys__: Set of keys that this fingerprint should _not_ match with.
-* __Metadata__: Arbitrary key/value data that can be aggregated for use in summary templates and for debugging.
-* __ChangeFilter__: A semicolon-delimited set of wildcards for the files that can be considered
-  responsible for this error, such as `*.c;*.cpp;*.h;*.hpp*` for C++ source files, for example.
+The `EpicGames.Horde` library provides extension methods using the standard .NET
+`ILogger` interface (exposed via `CommandUtils.Logger` in AutomationTool):
 
-### Matching and merging fingerprints
+```cs
+IssueFingerprint fingerprint = new IssueFingerprint(
+    "Compile",
+    "Compile {Severity} in {Files}",
+    IssueChangeFilter.Code);
+fingerprint.Keys.Add(IssueKey.FromFile("Foo.cpp"));
 
-Fingerprints are matched by logic in `IIssueFingerprint.cs`. The pipeline is as follows:
+using (IDisposable scope = Logger.BeginIssueScope(fingerprint))
+{
+    Logger.LogError("Compile errors in {File}", new FileReference("Foo.cpp"));
+}
+```
+
+#### Adding Fingerprints via Post-Processing
+
+Post-processing uses `IssueHandler` derived classes with the `[IssueHandler]`
+attribute. A new handler instance is created per build step. Log events are passed
+to `HandleEvent` in priority order until one returns `true`; then `GetIssues`
+returns the matched events and fingerprints.
+
+Handlers can return `true` without generating a fingerprint to mask vague errors
+(e.g. a failing exit code) when a more specific handler (e.g. compile error) has
+already claimed the root cause.
+
+### Issue Handlers
+
+Built-in handlers:
+
+| Handler | Description | Event Types |
+|---------|-------------|-------------|
+| `CompileIssueHandler` | Multi-line Clang, MSVC, static analyzer errors; groups by source file. | `Compiler`, `UHT`, `SourceFileLine`, `MSBuild` |
+| `ContentIssueHandler` | Parses `UE_ASSET_LOG` output; identifies modified assets. | `Engine_AssetLog` |
+| `CopyrightIssueHandler` | `CheckCopyrightNotices` UAT command output. | `AutomationTool_MissingCopyright` |
+| `DefaultIssueHandler` | Catch-all; groups unhandled events in the same step. | Any |
+| `GauntletIssueHandler` | Parses Gauntlet test failure names. | `Gauntlet_TestEvent`, `Gauntlet_DeviceEvent`, etc. |
+| `HashedIssueHandler` | Redacts numeric/timestamp patterns; groups by hash. | Any (>30 chars after redaction) |
+| `LocalizationIssueHandler` | UE localization command output. | `Engine_Localization` |
+| `PerforceCaseIssueHandler` | `CheckPerforceCase` UAT command output. | `AutomationTool_PerforceCase` |
+| `ScopedIssueHandler` | Uses scoped fingerprints in log events. Enable via `"Scoped"` tag. | Any (tagged) |
+| `ShaderIssueHandler` | Shader compilation output. | `Engine_ShaderCompiler` |
+| `SymbolIssueHandler` | Matches linker symbols to modified source files. | `Linker`, `Linker_DuplicateSymbol`, `Linker_UndefinedSymbol` |
+| `SystemicIssueHandler` | Infrastructure/flaky test issues. | `Systemic` through `Systemic_Max` |
+| `UnacceptableWordsIssueHandler` | `CheckUnacceptableWords` UAT command output. | `AutomationTool_UnacceptableWords` |
+| `SanitizerIssueHandler` | ASAN, TSAN, etc. output. | Sanitizer events |
+
+Workflows can opt into specific handlers via the [`issueHandlers`](#issue-handlers-1)
+property.
+
+### Summary Templates
 
-* A fingerprint is generated for each error or warning in a build step.
-* Each fingerprint is compared against the fingerprint of any unresolved spans in the same
-  stream. New fingerprints are considered to match an existing fingerprint if:
-    1. The type is identical
-    2. The summary template is identical
-    3. The new fingerprint has one key in common with the existing fingerprint, or neither fingerprint has any keys.
-    4. The new fingerprint does not contain any of the keys listed in the _rejectKeys_.
-* All the fingerprints which aren't matched with an existing span are grouped together and
-  added to new spans. These fingerprints are grouped using a more relaxed set of conditions:
-    1. The type is identical
-    2. Neither fingerprint has a key that is in other key's reject keys.
-* For any new spans, we find the list of suspect commits - following the merge history for that
-  commit to its originating stream. If any issue has a matching fingerprint, shares a suspect
-  commit, and does not yet have a span in the current stream, we add it to that issue. Otherwise a new issue is created.
+| Variable | Description |
+|----------|-------------|
+| `{Severity}` | `Warnings` or `Errors` based on highest severity. |
+| `{Files}` | Up to 3 filenames from `File`-type keys. |
+| `{Meta:Key}` | Metadata values for the given key. |
 
-### Creating fingerprints
+---
 
-Fingerprints may be created for errors in two ways:
+## Workflows
 
-* By including fingerprint information directly in the [structured log event](../Internals/StructuredLogging.md) at the
-  point that it is generated.
-* By post-processing [structured log events](../Internals/StructuredLogging.md) in the Horde server once a build step
-  completes.
+Workflows define how and where developers are notified about build health issues.
+They are configured in the `workflows` array of the
+[stream configuration](Schema/Streams.md) and assigned to jobs via the template's
+`workflowId` property. Individual BuildGraph nodes can override the workflow via
+[annotations](#annotations).
 
-Including fingerprints directly in structured log events is preferred, since it
-allows for more extensibility and gives the most control over adding fine-grained context.
+Settings cascade: **stream** > **workflow** (including default `annotations`) >
+**template** (`workflowId`, `triageChannel`, `updateIssues`) > **node**
+(BuildGraph annotations).
 
-#### Adding fingerprints to errors (C#)
+### Notification Lifecycle
 
-The `EpicGames.Horde` library includes extension methods for creating scopes that attach fingerprint
-metadata to log events, using the standard .NET `ILogger` interface. This interface is exposed to
-`AutomationTool` commands via the `CommandUtils.Logger` property.
+1. **New issue** — Thread created in `triageChannel` with `triagePrefix`.
+   `triageAlias` is @mentioned; `triageInstructions` posted as follow-up.
+2. **Suspects** — If `allowMentions` is true, up to `maxMentions` developers
+   are @mentioned.
+3. **Escalation** — After `escalateTimes` minutes, `escalateAlias` is pinged.
+4. **Reports** — At each `reportTimes`, a summary posts to `reportChannel`.
+5. **Resolution** — Thread updated when all spans resolve.
 
-Fingerprints can be added as in the following example:
+Slack parameters require **IDs** (channel ID, user/group ID), not display names.
 
-  ```cs
-  IssueFingerprint fingerprint = new IssueFingerprint("Compile", "Compile {Severity} in {Files}", IssueChangeFilter.Code);
-  fingerprint.Keys.Add(IssueKey.FromFile("Foo.cpp"));
+### Workflow Property Reference
 
-  using (IDisposable scope = Logger.BeginIssueScope(fingerprint))
-  {
-      Logger.LogError("Compile errors in {File}", new FileReference("Foo.cpp"));
-  }
-  ```
+#### Identity
 
-#### Adding fingerprints to errors (C++)
+| Property | Type | Default | Description |
+|----------|------|---------|-------------|
+| `id` | `string` | (required) | Unique identifier referenced by templates via `workflowId`. |
+| `summaryTab` | `string` | `null` | Dashboard tab name in the Build Health view. |
 
-// TODO
+#### Triage
 
-#### Adding fingerprints via post-processing
+| Property | Type | Default | Description |
+|----------|------|---------|-------------|
+| `triageChannel` | `string` | `null` | Slack channel ID for new issue threads. |
+| `triagePrefix` | `string` | `"*[NEW]* "` | Text prepended to triage messages. |
+| `triageSuffix` | `string` | `null` | Text appended to triage messages. |
+| `triageInstructions` | `string` | `null` | Developer guidance posted in the thread. Supports Slack markup. |
+| `triageAlias` | `string` | `null` | Slack user/group ID to @mention on new issues. |
+| `triageTypeAliases` | `object` | `null` | Map of issue type to Slack ID for type-specific routing (e.g. `{"Systemic": "SYYYYYYYYYY"}`). |
+| `triageErrors` | `bool` | `true` | Create triage threads for errors. |
+| `triageWarnings` | `bool` | `true` | Create triage threads for warnings. |
 
-Post-processing is handled by the `IssueHandler` derived classes in the `EpicGames.Horde` project.
+#### Reports
 
-Handlers are enumerated by the server by searching for classes with the `[IssueHandler]` attribute. A new
-instance of each handler is created for each build step containing events to be processed.
+| Property | Type | Default | Description |
+|----------|------|---------|-------------|
+| `reportTimes` | `string[]` | `[]` | UTC times for summary reports (`"HH:mm:ss"`). |
+| `reportChannel` | `string` | `null` | Slack channel ID for reports (can differ from triage). |
+| `reportWarnings` | `bool` | `true` | Include warnings in reports. |
+| `groupIssuesByTemplate` | `bool` | `true` | Group issues by job template in reports. |
+| `skipWhenEmpty` | `bool` | `false` | Skip report when no active issues. |
 
-Log events are passed to `IssueHandler.HandleEvent` in decreasing order of their `IssueHandler.Priority` value,
-until a handler is found which returns `true`. Once all events have been processed, `IssueHandler.GetIssues`
-returns a list of the matched log events and their fingerprints.
+#### Escalation
 
-Note that handlers do not have to generate fingerprints for a log event; it can be useful
-for more specific errors (say, a compile error) to mask more vague errors (say, UnrealBuildTool
-returning a failing exit code). As such, `HandleEvent` can return true for errors that it wants
-to exclude from the remaining handler pipeline.
+| Property | Type | Default | Description |
+|----------|------|---------|-------------|
+| `escalateAlias` | `string` | `null` | Slack user/group ID for escalation. |
+| `escalateTimes` | `int[]` | `[120]` | Minutes after creation to escalate. Multiple values = tiered escalation (e.g. `[120, 1440]`). |
+| `escalateRootCauseAnalysisTimes` | `int[]` | `[1440]` | Minutes to escalate issues missing an RCA. |
 
-### Summary Templates
+#### Mentions
 
-The following placeholder variables can be used in summary templates:
+| Property | Type | Default | Description |
+|----------|------|---------|-------------|
+| `allowMentions` | `bool` | `true` | @mention suspected developers. Set `false` for info-only channels. |
+| `maxMentions` | `int` | `5` | Max users to @mention per issue. |
+| `inviteRestrictedUsers` | `bool` | `false` | Invite external/contractor Slack users to channels. |
 
-* `{Severity}`: Expands to the string 'Warnings' or 'Errors' depending on
-  the highest severity of diagnostic grouped into this issue.
-* `{Files}`: Expands to a list of filenames from keys with a type of `File`. Up
-  to 3 files are included.
-* `{Meta:Key}`: Expands to a list of metadata values associated with `Key`.
+#### External Issue Tracking
 
-## Workflows
+| Property | Type | Default | Description |
+|----------|------|---------|-------------|
+| `externalIssues.projectKey` | `string` | `null` | Project key in external tracker (e.g. Jira). `""` disables creation. |
+| `externalIssues.defaultComponentId` | `string` | `""` | Default component for created issues. |
+| `externalIssues.defaultIssueTypeId` | `string` | `""` | Default issue type for created issues. |
+
+#### Other
+
+| Property | Type | Default | Description |
+|----------|------|---------|-------------|
+| `issueHandlers` | `string[]` | `null` | Handler tags to enable (e.g. `["Scoped"]`). `null` = all defaults. |
+| `autoRcaEnabled` | `bool` | `false` | Enable automatic root cause analysis. |
+| `showMergeWarnings` | `bool` | `false` | Show warnings for issues introduced via merges. |
+| `annotations` | `object` | `{}` | Default [node annotations](#annotations) for all nodes in this workflow. |
+
+---
+
+## Workflow Examples
+
+### Minimal
+
+Posts triage notifications to a single channel:
+
+```json
+{
+    "workflows": [
+        {
+            "id": "default",
+            "triageChannel": "CXXXXXXXXXX",
+            "triagePrefix": "*[CI]*"
+        }
+    ]
+}
+```
+
+### CI with Reports and External Issues
+
+Adds scheduled reports and external tracker integration:
+
+```json
+{
+    "workflows": [
+        {
+            "id": "ci-build-health",
+            "summaryTab": "CI Build Health",
+            "reportTimes": ["09:00:00", "21:00:00"],
+            "reportChannel": "CXXXXXXXXXX",
+            "triageChannel": "CXXXXXXXXXX",
+            "triagePrefix": "*[CI]*",
+            "externalIssues": { "projectKey": "MYPROJECT" },
+            "allowMentions": false
+        }
+    ]
+}
+```
 
-Workflows define how and where developers are notified about build health
-issues.
+### High-Priority with Escalation and Auto-RCA
 
-Workflows are configured via the [`Workflows`](../Config/Schema/Streams.md) property in
-the stream configuration file, and jobs are configured to use them through the job
-template's [`WorkflowId`](../Config/Schema/Streams.md#templaterefconfig) property.
+For the main development branch where breakages must be fixed immediately:
 
-Individual nodes within a job may be configured to use a workflow using an [annotation](#annotations) in the BuildGraph script.
+```json
+{
+    "workflows": [
+        {
+            "id": "high-priority-ci",
+            "summaryTab": "High Priority CIS",
+            "reportTimes": ["03:00:00", "09:00:00", "15:00:00", "21:00:00"],
+            "reportChannel": "CXXXXXXXXXX",
+            "triageChannel": "CYYYYYYYYYY",
+            "triagePrefix": "*[Main CIS]*",
+            "triageInstructions": "If you caused this issue, it needs to be backed out immediately. Please acknowledge and fix as soon as possible.",
+            "triageAlias": "SXXXXXXXXXX",
+            "triageTypeAliases": {
+                "Systemic": "SYYYYYYYYYY"
+            },
+            "escalateAlias": "UYYYYYYYYYY",
+            "escalateTimes": [120, 1440],
+            "externalIssues": { "projectKey": "MYPROJECT" },
+            "allowMentions": true,
+            "inviteRestrictedUsers": true,
+            "autoRcaEnabled": true
+        }
+    ]
+}
+```
 
-### Notifications
+Key features demonstrated:
+* `triageTypeAliases` routes `Systemic` issues to a dedicated infra team.
+* Two-tier escalation: 2 hours, then 24 hours.
+* `autoRcaEnabled` activates automatic root cause analysis.
 
-Horde supports [Slack](https://www.slack.com) for issue notifications. By default, a
-Slack thread is created for each identified issue, and developers that are suspected of
-causing the issue are @mentioned to draw their attention to it.
+### Low-Noise with Skip-When-Empty
 
-An interactive surface allows developers to acknowledge or decline issues from Slack
-without having to navigate to the Horde dashboard.
+For secondary streams where breakages are infrequent:
 
-When configuring Slack channels and users through the stream configuration file, most
-parameters require the Slack channel and user ID to be supplied rather than their name.
+```json
+{
+    "workflows": [
+        {
+            "id": "dev-stream-health",
+            "summaryTab": "Dev Stream",
+            "reportTimes": ["06:00:00"],
+            "reportChannel": "CXXXXXXXXXX",
+            "triageChannel": "CXXXXXXXXXX",
+            "externalIssues": { "projectKey": "MYPROJECT" },
+            "allowMentions": true,
+            "skipWhenEmpty": true
+        }
+    ]
+}
+```
 
-### Reports
+### Multiple Workflows per Stream
 
-Workflows can be configured to send out summary reports at defined intervals, showing a
-list of unresolved issues, links to the appropriate threads, and statistics for the
-health of the stream.
+A real-world pattern: different job categories route to different channels and teams.
 
-The sending of reports can be configured through the `reportTimes` and `reportChannel` properties of the workflow configuration.
+```json
+{
+    "workflows": [
+        {
+            "id": "ci-high-priority",
+            "summaryTab": "High Priority CIS",
+            "reportTimes": ["09:00:00", "15:00:00", "21:00:00"],
+            "reportChannel": "CXXXXXXXXXX",
+            "triageChannel": "CYYYYYYYYYY",
+            "triagePrefix": "*[Main CIS]*",
+            "triageInstructions": "Back out immediately if you caused this.",
+            "escalateAlias": "SXXXXXXXXXX",
+            "escalateTimes": [120],
+            "externalIssues": { "projectKey": "MYPROJECT" },
+            "allowMentions": true,
+            "inviteRestrictedUsers": true
+        },
+        {
+            "id": "scheduled-builds",
+            "summaryTab": "Scheduled Builds",
+            "reportTimes": ["08:00:00", "14:00:00"],
+            "reportChannel": "CZZZZZZZZZZ",
+            "triageChannel": "CWWWWWWWWWW",
+            "triagePrefix": "*[Scheduled Build]*",
+            "escalateTimes": [480, 1440],
+            "externalIssues": { "projectKey": "MYPROJECT" },
+            "allowMentions": false
+        },
+        {
+            "id": "qa-tracking",
+            "summaryTab": "QA Tracking",
+            "reportTimes": ["09:00:00", "18:00:00"],
+            "triageChannel": "CVVVVVVVVVV",
+            "triagePrefix": "*[QA]*",
+            "triageWarnings": false,
+            "reportWarnings": false,
+            "issueHandlers": ["Scoped"],
+            "allowMentions": false
+        }
+    ]
+}
+```
 
-### Annotations
+Templates reference workflows by ID. Set `updateIssues: true` on templates that
+are not regular CI but should still participate in build health tracking:
 
-Handling of issues for individual job steps may be further configured using __node attributes__. Attributes can be
-specified for a node through the `Annotations` attribute on `Node` elements in BuildGraph scripts:
+```json
+{
+    "templates": [
+        {
+            "id": "incremental-ci",
+            "name": "Incremental CI",
+            "workflowId": "ci-high-priority",
+            "arguments": ["-Target=CI Targets", "-Script=Build/CI.xml"]
+        },
+        {
+            "id": "nightly-build",
+            "name": "Nightly Build",
+            "workflowId": "scheduled-builds",
+            "updateIssues": true,
+            "arguments": ["-Target=Full Build", "-Script=Build/FullBuild.xml"]
+        }
+    ]
+}
+```
+
+---
+
+## Annotations
+
+Issue handling for individual job steps can be configured via **node annotations**
+in BuildGraph scripts:
+
+```xml
+<Node Name="Compile Editor Win64" Annotations="Workflow=my-workflow;BuildBlocker=true">
+```
 
-  ```xml
-  <Node Name="Compile UnrealEditor Win64" Annotations="Workflow=my-workflow;BuildBlocker=true">
-  ```
+| Annotation | Type | Description |
+|------------|------|-------------|
+| `Workflow` | `string` | Workflow ID for this node. Overrides template-level `workflowId`. |
+| `CreateIssues` | `bool` | Set `false` to disable issue creation. |
+| `CreateWarningIssues` | `bool` | Set `false` to disable warning-level issues. |
+| `AutoAssign` | `bool` | Auto-assign when only one suspect or a clear file correlation exists. |
+| `AutoAssignToUser` | `string` | Always assign issues from this step to the given user. |
+| `NotifySubmitters` | `bool` | Notify all submitters between success and failure. |
+| `IssueGroup` | `string` | Suffix for the issue type, preventing merging with other issues. |
+| `BuildBlocker` | `bool` | Flag failures as build blockers (shown prominently in Slack). |
 
-The supported set of annotations are defined in the `NodeAnnotations.cs` source file. The following are currently available:
+Multiple annotations are semicolon-separated. Annotations can also be set as
+defaults in a workflow's `annotations` property.
 
-* `Workflow`: Workflow to use for triaging issues from this node
-* `CreateIssues`: Allows disabling creation of issues for this node, when set to `false`.
-* `CreateWarningIssues`: Allows disabling creation of warning issues for this node, when set to `false`.
-* `AutoAssign`: Whether to automatically assign issues that could only be caused by one user, or have a well defined
-  correlation with a modified file.
-* `AutoAssignToUser`: Automatically assign any issues in this step to the Perforce user given by the supplied argument.
-* `NotifySubmitters`: Whether to notify all submitters between a build succeeding and failing, allowing them to step
-  forward and take ownership of an issue.
-* `IssueGroup`: Specifies a suffix to be appended to the `type` property of any issues created from this node,
-  preventing them being merged with other issues.
-* `BuildBlocker`: Whether failures in this node should be considered a build blocker. Issues identified as build
-  blockers have a special tag in Slack notifications.
+---
 
 ## Ignore Patterns
 
-Some build steps produce warnings or errors that are known, expected, or otherwise
-not actionable. Without intervention, these will create spurious Build Health issues
-and notify developers unnecessarily.
+**Ignore patterns** suppress known, non-actionable warnings or errors from Build
+Health. Matched lines are demoted to informational severity — they still appear in
+logs but won't create issues or notifications.
 
-__Ignore patterns__ let you suppress specific log lines so they are not ingested by the
-Build Health system. When a warning or error matches an ignore pattern, it is
-re-logged as an informational [structured log event](../Internals/StructuredLogging.md)
-with the following properties attached:
+Each suppressed line is tagged with:
 
 | Property | Description |
 |----------|-------------|
-| `IgnoredLine` | The original text of the suppressed line |
-| `IgnorePatternFile` | Absolute path to the file that defined the matching pattern |
-| `IgnorePattern` | The regular expression that matched |
-
-Because the line is demoted to informational severity, it will not generate Build
-Health issues or trigger notifications, but it still appears in the full build log.
-The attached properties make it straightforward to trace __why__ a particular line
-was suppressed and __which pattern file__ was responsible — useful when debugging
-unexpected suppressions or reviewing which patterns are active.
+| `IgnoredLine` | Original text of the suppressed line |
+| `IgnorePatternFile` | Path to the pattern file that matched |
+| `IgnorePattern` | The regex that matched |
 
 ### File Format
 
-An ignore patterns file contains one regular expression per line. Lines starting with
-`#` are treated as comments and blank lines are ignored.
+One .NET regex per line. `#` lines are comments; blank lines are ignored.
 
 ```
-# IgnorePatterns.txt - suppress known warnings for this project
-# ──────────────────────────────────────────────────────────────
-
-# Third-party SDK deprecation warnings we can't fix until the next SDK update
+# Suppress known third-party deprecation warnings
 warning C4996: 'ThirdPartyLib::DeprecatedFunc'.*was declared deprecated
 
-# Intermittent linker note on Mac that is not a real error
+# Intermittent linker note on Mac
 ld: warning: object file.*was built for newer macOS version
 
-# Shader compiler warnings for experimental feature behind a feature flag
+# Experimental shader feature behind a flag
 LogShaderCompiler: Warning:.*EXPERIMENTAL_FEATURE
 ```
 
-Each line is compiled as a .NET regular expression, so standard regex syntax
-applies. A 200ms timeout is enforced per match to guard against catastrophic
-backtracking, so prefer simple, specific patterns.
+A 200ms timeout is enforced per match to guard against catastrophic backtracking.
 
 ### File Discovery
 
-Horde agents automatically scan the workspace for files named
-`Build/Horde/IgnorePatterns.txt` in standard locations:
+Agents scan for `Build/Horde/IgnorePatterns.txt` in standard locations:
 
 | Location | Example |
 |----------|---------|
 | Engine | `Engine/Build/Horde/IgnorePatterns.txt` |
-| Top-level project directories | `MyProject/Build/Horde/IgnorePatterns.txt` |
+| Project directories | `MyProject/Build/Horde/IgnorePatterns.txt` |
 | Platform subdirectories | `Engine/Platforms/Linux/Build/Horde/IgnorePatterns.txt` |
 
-All discovered files are merged together; patterns from every file are active for
-the job.
-
-### Additional Pattern Files via Configuration
-
-You can reference extra pattern files from your project or stream configuration
-using the `IgnorePatternFiles` property. Paths are relative to the workspace root.
-
-In a __project__ config (applies to all streams in the project):
+All discovered files are merged. Additional files can be referenced via
+`ignorePatternFiles` in project or stream config:
 
 ```json
 {
@@ -288,15 +479,5 @@ In a __project__ config (applies to all streams in the project):
 }
 ```
 
-In a __stream__ config (appended to any project-level entries):
-
-```json
-{
-    "ignorePatternFiles": [
-        "MyProject/Config/Horde/StreamSpecificIgnores.txt"
-    ]
-}
-```
-
-See the [ProjectConfig](Schema/Projects.md) and [StreamConfig](Schema/Streams.md)
-references for the full schema.
+See [ProjectConfig](Schema/Projects.md) and [StreamConfig](Schema/Streams.md)
+for the full schema.
