Horde: Update Build Health documentation with comprehensive workflow reference and examples
#rnx

diff --git a/Engine/Source/Programs/Horde/Docs/Config/BuildHealth.md b/Engine/Source/Programs/Horde/Docs/Config/BuildHealth.md
--- a/Engine/Source/Programs/Horde/Docs/Config/BuildHealth.md
+++ b/Engine/Source/Programs/Horde/Docs/Config/BuildHealth.md
@@ -131,8 +131,8 @@ Built-in handlers:
 | `UnacceptableWordsIssueHandler` | `CheckUnacceptableWords` UAT command output. | `AutomationTool_UnacceptableWords` |
 | `SanitizerIssueHandler` | ASAN, TSAN, etc. output. | Sanitizer events |
 
-Workflows can opt into specific handlers via the [`issueHandlers`](#issue-handlers-1)
-property.
+Workflows can opt into specific handlers via the `issueHandlers`
+[workflow property](#other).
 
 ### Summary Templates
 
