Horde: Updating schema documentation

#rnx

diff --git a/Engine/Source/Programs/Horde/Docs/Config/Schema/AclActions.md b/Engine/Source/Programs/Horde/Docs/Config/Schema/AclActions.md
--- a/Engine/Source/Programs/Horde/Docs/Config/Schema/AclActions.md
+++ b/Engine/Source/Programs/Horde/Docs/Config/Schema/AclActions.md
@@ -37,44 +37,6 @@
 | `DeleteAccount` |  |
 | `ViewAccount` |  |
 
-## Agents
-
-| Name | Description |
-| ---- | ----------- |
-| `CreateAgent` |  |
-| `CreateWorkstationAgent` |  |
-| `UpdateAgent` |  |
-| `DeleteAgent` |  |
-| `ViewAgent` |  |
-| `ListAgents` |  |
-
-## Compute
-
-| Name | Description |
-| ---- | ----------- |
-| `AddComputeTasks` |  |
-| `GetComputeTasks` |  |
-| `UbaCacheRead` |  |
-| `UbaCacheWrite` |  |
-
-## Leases
-
-| Name | Description |
-| ---- | ----------- |
-| `ViewLeases` |  |
-| `ViewLeaseTasks` |  |
-
-## Logs
-
-| Name | Description |
-| ---- | ----------- |
-| `CreateLog` |  |
-| `UpdateLog` |  |
-| `ViewLog` |  |
-| `WriteLogData` |  |
-| `CreateEvent` |  |
-| `ViewEvent` |  |
-
 ## Notices
 
 | Name | Description |
@@ -83,16 +45,6 @@
 | `UpdateNotice` |  |
 | `DeleteNotice` |  |
 
-## Pools
-
-| Name | Description |
-| ---- | ----------- |
-| `CreatePool` |  |
-| `UpdatePool` |  |
-| `DeletePool` |  |
-| `ViewPool` |  |
-| `ListPools` |  |
-
 ## ServiceAccounts
 
 | Name | Description |
@@ -102,21 +54,6 @@
 | `DeleteAccount` |  |
 | `ViewAccount` |  |
 
-## Sessions
-
-| Name | Description |
-| ---- | ----------- |
-| `CreateSession` |  |
-| `ViewSession` |  |
-
-## Software
-
-| Name | Description |
-| ---- | ----------- |
-| `UploadSoftware` |  |
-| `DownloadSoftware` |  |
-| `DeleteSoftware` |  |
-
 ## Accounts
 
 | Name | Description |
@@ -848,6 +785,12 @@
 | `UpdateReplicator` |  |
 | `ViewReplicator` |  |
 
+## Secrets
+
+| Name | Description |
+| ---- | ----------- |
+| `ViewSecret` |  |
+
 ## ServiceAccounts
 
 | Name | Description |
@@ -883,12 +826,6 @@
 | `ViewChanges` |  |
 | `ViewTemplate` |  |
 
-## Symbols
-
-| Name | Description |
-| ---- | ----------- |
-| `ReadSymbols` |  |
-
 ## TestData
 
 | Name | Description |
@@ -1025,6 +962,12 @@
 | `UpdateReplicator` |  |
 | `ViewReplicator` |  |
 
+## Secrets
+
+| Name | Description |
+| ---- | ----------- |
+| `ViewSecret` |  |
+
 ## ServiceAccounts
 
 | Name | Description |
@@ -1049,6 +992,19 @@
 | `DownloadSoftware` |  |
 | `DeleteSoftware` |  |
 
+## Storage
+
+| Name | Description |
+| ---- | ----------- |
+| `ReadBlobs` |  |
+| `WriteBlobs` |  |
+| `ReadAliases` |  |
+| `WriteAliases` |  |
+| `DeleteAliases` |  |
+| `ReadRefs` |  |
+| `WriteRefs` |  |
+| `DeleteRefs` |  |
+
 ## Streams
 
 | Name | Description |
@@ -1060,12 +1016,6 @@
 | `ViewChanges` |  |
 | `ViewTemplate` |  |
 
-## Symbols
-
-| Name | Description |
-| ---- | ----------- |
-| `ReadSymbols` |  |
-
 ## TestData
 
 | Name | Description |
@@ -1073,13 +1023,6 @@
 | `TestHealthRead` |  |
 | `TestHealthWrite` |  |
 
-## Tools
-
-| Name | Description |
-| ---- | ----------- |
-| `DownloadTool` |  |
-| `UploadTool` |  |
-
 ## Accounts
 
 | Name | Description |
@@ -1239,6 +1182,19 @@
 | `DownloadSoftware` |  |
 | `DeleteSoftware` |  |
 
+## Storage
+
+| Name | Description |
+| ---- | ----------- |
+| `ReadBlobs` |  |
+| `WriteBlobs` |  |
+| `ReadAliases` |  |
+| `WriteAliases` |  |
+| `DeleteAliases` |  |
+| `ReadRefs` |  |
+| `WriteRefs` |  |
+| `DeleteRefs` |  |
+
 ## Streams
 
 | Name | Description |
@@ -1263,13 +1219,6 @@
 | `TestHealthRead` |  |
 | `TestHealthWrite` |  |
 
-## Tools
-
-| Name | Description |
-| ---- | ----------- |
-| `DownloadTool` |  |
-| `UploadTool` |  |
-
 ## Accounts
 
 | Name | Description |
diff --git a/Engine/Source/Programs/Horde/Docs/Config/Schema/Globals.md b/Engine/Source/Programs/Horde/Docs/Config/Schema/Globals.md
--- a/Engine/Source/Programs/Horde/Docs/Config/Schema/Globals.md
+++ b/Engine/Source/Programs/Horde/Docs/Config/Schema/Globals.md
@@ -55,15 +55,21 @@
 
 Name | Description
 ---- | -----------
-`compute` | [ComputeConfig](#computeconfig)<br>Configuration for the compute plugin
 `analytics` | [AnalyticsConfig](#analyticsconfig)<br>Configuration for the analytics plugin
+`compute` | [ComputeConfig](#computeconfig)<br>Configuration for the compute plugin
 `build` | [BuildConfig](#buildconfig)<br>Configuration for the build plugin
-`exampleplugin` | [ExamplePluginConfig](#examplepluginconfig)<br>Configuration for the exampleplugin plugin
 `experimental` | [ExperimentalConfig](#experimentalconfig)<br>Configuration for the experimental plugin
-`symbols` | [SymbolsConfig](#symbolsconfig)<br>Configuration for the symbols plugin
-`tools` | [ToolsConfig](#toolsconfig)<br>Configuration for the tools plugin
+`exampleplugin` | [ExamplePluginConfig](#examplepluginconfig)<br>Configuration for the exampleplugin plugin
 `secrets` | [SecretsConfig](#secretsconfig)<br>Configuration for the secrets plugin
 `storage` | [StorageConfig](#storageconfig)<br>Configuration for the storage plugin
+`symbols` | [SymbolsConfig](#symbolsconfig)<br>Configuration for the symbols plugin
+`tools` | [ToolsConfig](#toolsconfig)<br>Configuration for the tools plugin
 
+## AnalyticsConfig
+
+Name | Description
+---- | -----------
+`stores` | [TelemetryStoreConfig](Telemetry.md)`[]`<br>
+
 ## ComputeConfig
 
@@ -265,12 +271,6 @@
 `description` | `string`<br>
 `computeId` | `string`<br>
 
-## AnalyticsConfig
-
-Name | Description
----- | -----------
-`stores` | [TelemetryStoreConfig](Telemetry.md)`[]`<br>
-
 ## BuildConfig
 
 Name | Description
@@ -387,56 +387,17 @@
 `keepDays` | `integer`<br>
 `namespaceId` | `string`<br>
 
-## ExamplePluginConfig
-
-Name | Description
----- | -----------
-`greetingMessage` | `string`<br>
-
 ## ExperimentalConfig
 
 Name | Description
 ---- | -----------
 `notifications` | [NotificationConfig](Notification.md)`[]`<br>
 
-## SymbolsConfig
+## ExamplePluginConfig
 
 Name | Description
 ---- | -----------
-`stores` | [SymbolStoreConfig](#symbolstoreconfig)`[]`<br>
+`greetingMessage` | `string`<br>
 
-## SymbolStoreConfig
-
-Name | Description
----- | -----------
-`id` | `string`<br>
-`namespaceId` | `string`<br>
-`public` | `boolean`<br>
-`acl` | [AclConfig](#aclconfig)<br>
-
-## ToolsConfig
-
-Name | Description
----- | -----------
-`tools` | [ToolConfig](#toolconfig)`[]`<br>
-
-## ToolConfig
-
-Name | Description
----- | -----------
-`id` | `string`<br>
-`name` | `string`<br>
-`description` | `string`<br>
-`category` | `string`<br>
-`group` | `string`<br>
-`platforms` | `string[]`<br>
-`public` | `boolean`<br>
-`showInUgs` | `boolean`<br>
-`showInDashboard` | `boolean`<br>
-`showInToolbox` | `boolean`<br>
-`metadata` | `string` `->` `string`<br>
-`namespaceId` | `string`<br>
-`acl` | [AclConfig](#aclconfig)<br>
-
 ## SecretsConfig
 
@@ -565,3 +526,42 @@
 `gcDelayHrs` | `number`<br>
 `enableAliases` | `boolean`<br>
 `acl` | [AclConfig](#aclconfig)<br>
+
+## SymbolsConfig
+
+Name | Description
+---- | -----------
+`stores` | [SymbolStoreConfig](#symbolstoreconfig)`[]`<br>
+
+## SymbolStoreConfig
+
+Name | Description
+---- | -----------
+`id` | `string`<br>
+`namespaceId` | `string`<br>
+`public` | `boolean`<br>
+`acl` | [AclConfig](#aclconfig)<br>
+
+## ToolsConfig
+
+Name | Description
+---- | -----------
+`tools` | [ToolConfig](#toolconfig)`[]`<br>
+
+## ToolConfig
+
+Name | Description
+---- | -----------
+`id` | `string`<br>
+`name` | `string`<br>
+`description` | `string`<br>
+`category` | `string`<br>
+`group` | `string`<br>
+`platforms` | `string[]`<br>
+`public` | `boolean`<br>
+`showInUgs` | `boolean`<br>
+`showInDashboard` | `boolean`<br>
+`showInToolbox` | `boolean`<br>
+`metadata` | `string` `->` `string`<br>
+`namespaceId` | `string`<br>
+`acl` | [AclConfig](#aclconfig)<br>
