Horde: Improve truncation of over-sized log messages

Preserve RenderedMessage in truncated messages, and Exception as well if it fits.

diff --git a/Engine/Source/Programs/Horde/HordeServer.Tests/Utilities/SizeLimitedJsonFormatterTests.cs b/Engine/Source/Programs/Horde/HordeServer.Tests/Utilities/SizeLimitedJsonFormatterTests.cs
--- a/Engine/Source/Programs/Horde/HordeServer.Tests/Utilities/SizeLimitedJsonFormatterTests.cs
+++ b/Engine/Source/Programs/Horde/HordeServer.Tests/Utilities/SizeLimitedJsonFormatterTests.cs
@@ -2,6 +2,7 @@
 
 using System;
 using System.IO;
+using System.Text.Json;
 using HordeServer.Utilities;
 using Microsoft.VisualStudio.TestTools.UnitTesting;
 using Serilog.Events;
@@ -16,10 +17,20 @@ namespace HordeServer.Tests.Utilities
 		readonly JsonFormatter _inner = new JsonFormatter(renderMessage: true);
 		readonly SizeLimitedJsonFormatter _truncating = new SizeLimitedJsonFormatter(new JsonFormatter(renderMessage: true), maxBytes: 100);
 
-		static LogEvent CreateLogEvent(string messageTemplate, Exception? exception = null, LogEventLevel level = LogEventLevel.Information)
+		static LogEvent CreateLogEvent(string messageTemplate, Exception? exception = null, LogEventLevel level = LogEventLevel.Information, params LogEventProperty[] properties)
 		{
 			MessageTemplateParser parser = new MessageTemplateParser();
-			return new LogEvent(DateTimeOffset.UtcNow, level, exception, parser.Parse(messageTemplate), Array.Empty<LogEventProperty>());
+			return new LogEvent(DateTimeOffset.UtcNow, level, exception, parser.Parse(messageTemplate), properties);
+		}
+
+		/// <summary>
+		/// Creates a log event with ExceptionInfo and ExceptionDetail properties, simulating the Serilog.Exceptions enricher.
+		/// </summary>
+		static LogEvent CreateLogEventWithExceptionProperties(string messageTemplate, Exception exception, string exceptionInfoValue, string exceptionDetailValue)
+		{
+			LogEventProperty infoProperty = new LogEventProperty("ExceptionInfo", new ScalarValue(exceptionInfoValue));
+			LogEventProperty detailProperty = new LogEventProperty("ExceptionDetail", new ScalarValue(exceptionDetailValue));
+			return CreateLogEvent(messageTemplate, exception, LogEventLevel.Error, infoProperty, detailProperty);
 		}
 
 		static string Format(Serilog.Formatting.ITextFormatter formatter, LogEvent logEvent)
@@ -47,7 +58,7 @@ namespace HordeServer.Tests.Utilities
 			Assert.IsTrue(result.Contains("\"_truncated\":true", StringComparison.Ordinal), "Should contain _truncated marker");
 			Assert.IsTrue(result.Contains("\"Level\":\"Information\"", StringComparison.Ordinal), "Should contain level");
 			Assert.IsTrue(result.Contains("\"MessageTemplate\":", StringComparison.Ordinal), "Should contain message template");
-			Assert.IsFalse(result.Contains("\"RenderedMessage\"", StringComparison.Ordinal), "Should not contain full rendered output");
+			Assert.IsTrue(result.Contains("\"RenderedMessage\":", StringComparison.Ordinal), "Should contain rendered message");
 			Assert.IsFalse(result.Contains("ExceptionType", StringComparison.Ordinal), "Should not contain exception fields when no exception");
 			Assert.IsFalse(result.Contains("ExceptionMessage", StringComparison.Ordinal), "Should not contain exception fields when no exception");
 		}
@@ -62,6 +73,8 @@ namespace HordeServer.Tests.Utilities
 			Assert.IsTrue(result.Contains("\"_truncated\":true", StringComparison.Ordinal), "Should contain _truncated marker");
 			Assert.IsTrue(result.Contains("\"ExceptionType\":\"System.InvalidOperationException\"", StringComparison.Ordinal), "Should contain exception type");
 			Assert.IsTrue(result.Contains("\"ExceptionMessage\":\"Something went wrong\"", StringComparison.Ordinal), "Should contain exception message");
+			Assert.IsTrue(result.Contains("\"Exception\":", StringComparison.Ordinal), "Should contain full exception stack trace");
+			Assert.IsTrue(result.Contains("\"RenderedMessage\":", StringComparison.Ordinal), "Should contain rendered message");
 		}
 
 		[TestMethod]
@@ -72,12 +85,124 @@ namespace HordeServer.Tests.Utilities
 			string result = Format(_truncating, logEvent).TrimEnd();
 
 			// Verify it parses as valid JSON
-			System.Text.Json.JsonDocument.Parse(result);
+			JsonDocument.Parse(result);
 
 			// Verify special characters were escaped
 			Assert.IsFalse(result.Contains("\n\"", StringComparison.Ordinal), "Newlines in values should be escaped");
 			Assert.IsTrue(result.Contains("\\n", StringComparison.Ordinal), "Should contain escaped newlines");
 			Assert.IsTrue(result.Contains("\\\"quotes\\\"", StringComparison.Ordinal), "Should contain escaped quotes");
+			Assert.IsTrue(result.Contains("\"RenderedMessage\":", StringComparison.Ordinal), "Should contain rendered message");
+		}
+
+		[TestMethod]
+		public void Tier2StripsEnricherPropertiesAndFits()
+		{
+			InvalidOperationException exception = new InvalidOperationException("Oops");
+			string bulkyInfo = new string('X', 5000);
+			string bulkyDetail = new string('Y', 5000);
+			LogEvent logEvent = CreateLogEventWithExceptionProperties("Op failed", exception, bulkyInfo, bulkyDetail);
+
+			// Inner formatted output includes the bulky properties
+			string innerResult = Format(_inner, logEvent);
+			int innerByteCount = System.Text.Encoding.UTF8.GetByteCount(innerResult);
+
+			// Set maxBytes between stripped size (without ExceptionInfo/ExceptionDetail) and full size
+			// The stripped version should be much smaller than innerByteCount
+			int maxBytes = innerByteCount - 8000; // removing ~10K of bulky props should make it fit
+			SizeLimitedJsonFormatter formatter = new SizeLimitedJsonFormatter(_inner, maxBytes: maxBytes);
+			string result = Format(formatter, logEvent);
+
+			Assert.IsTrue(result.Contains("\"_truncated\":true", StringComparison.Ordinal), "Should contain _truncated marker");
+			Assert.IsTrue(result.Contains("\"_originalBytes\":", StringComparison.Ordinal), "Should contain original byte count");
+			Assert.IsTrue(result.Contains("\"Properties\":", StringComparison.Ordinal), "Should preserve Properties object");
+			Assert.IsTrue(result.Contains("\"RenderedMessage\":", StringComparison.Ordinal), "Should preserve RenderedMessage");
+			Assert.IsTrue(result.Contains("\"Exception\":", StringComparison.Ordinal), "Should preserve Exception field");
+			Assert.IsFalse(result.Contains("ExceptionInfo", StringComparison.Ordinal), "Should have stripped ExceptionInfo");
+			Assert.IsFalse(result.Contains("ExceptionDetail", StringComparison.Ordinal), "Should have stripped ExceptionDetail");
+
+			// Verify valid JSON
+			JsonDocument.Parse(result.TrimEnd());
+		}
+
+		[TestMethod]
+		public void Tier2FallsToTier3WhenStillTooLarge()
+		{
+			InvalidOperationException exception = new InvalidOperationException("Oops");
+			string bulkyInfo = new string('X', 1000);
+			string bulkyDetail = new string('Y', 1000);
+			LogEvent logEvent = CreateLogEventWithExceptionProperties("Op failed", exception, bulkyInfo, bulkyDetail);
+
+			// Set maxBytes very small so even stripped Tier 2 output won't fit
+			SizeLimitedJsonFormatter formatter = new SizeLimitedJsonFormatter(_inner, maxBytes: 100);
+			string result = Format(formatter, logEvent);
+
+			Assert.IsTrue(result.Contains("\"_truncated\":true", StringComparison.Ordinal), "Should contain _truncated marker");
+			Assert.IsTrue(result.Contains("\"RenderedMessage\":", StringComparison.Ordinal), "Should contain rendered message");
+			Assert.IsTrue(result.Contains("\"Exception\":", StringComparison.Ordinal), "Should contain full exception");
+			Assert.IsFalse(result.Contains("\"Properties\":", StringComparison.Ordinal), "Tier 3 should not contain Properties object");
+
+			// Verify valid JSON
+			JsonDocument.Parse(result.TrimEnd());
+		}
+
+		[TestMethod]
+		public void InnerExceptionIsPreservedInTier3()
+		{
+			Exception inner = new InvalidOperationException("Inner cause");
+			Exception outer = new AggregateException("Outer wrapper", inner);
+			LogEvent logEvent = CreateLogEvent("Multi-layered failure", outer, LogEventLevel.Error);
+			string result = Format(_truncating, logEvent);
+
+			Assert.IsTrue(result.Contains("\"Exception\":", StringComparison.Ordinal), "Should contain Exception field");
+
+			// Extract the Exception value and verify it contains both exception types
+			JsonDocument doc = JsonDocument.Parse(result.TrimEnd());
+			string exceptionText = doc.RootElement.GetProperty("Exception").GetString()!;
+			Assert.IsTrue(exceptionText.Contains("AggregateException", StringComparison.Ordinal), "Should contain outer exception type");
+			Assert.IsTrue(exceptionText.Contains("InvalidOperationException", StringComparison.Ordinal), "Should contain inner exception type");
+			Assert.IsTrue(exceptionText.Contains("Inner cause", StringComparison.Ordinal), "Should contain inner exception message");
+		}
+
+		[TestMethod]
+		public void NoEnricherPropertiesSkipsTier2()
+		{
+			// Exception without ExceptionInfo/ExceptionDetail should go straight to Tier 3
+			InvalidOperationException exception = new InvalidOperationException("Plain exception");
+			LogEvent logEvent = CreateLogEvent("Failed operation", exception, LogEventLevel.Error);
+
+			// Use a small limit that forces truncation
+			SizeLimitedJsonFormatter formatter = new SizeLimitedJsonFormatter(_inner, maxBytes: 100);
+			string result = Format(formatter, logEvent);
+
+			Assert.IsTrue(result.Contains("\"_truncated\":true", StringComparison.Ordinal), "Should contain _truncated marker");
+			Assert.IsTrue(result.Contains("\"RenderedMessage\":", StringComparison.Ordinal), "Should contain rendered message");
+			Assert.IsTrue(result.Contains("\"Exception\":", StringComparison.Ordinal), "Should contain full exception");
+			Assert.IsTrue(result.Contains("\"ExceptionType\":\"System.InvalidOperationException\"", StringComparison.Ordinal), "Should contain exception type");
+			Assert.IsFalse(result.Contains("\"Properties\":", StringComparison.Ordinal), "Should be Tier 3 output without Properties");
+
+			// Verify valid JSON
+			JsonDocument.Parse(result.TrimEnd());
+		}
+
+		[TestMethod]
+		public void RenderedMessageContainsResolvedVariables()
+		{
+			LogEventProperty userProp = new LogEventProperty("User", new ScalarValue("Alice"));
+			LogEventProperty actionProp = new LogEventProperty("Action", new ScalarValue("deploy"));
+			LogEvent logEvent = CreateLogEvent("User {User} performed {Action}", null, LogEventLevel.Information, userProp, actionProp);
+
+			SizeLimitedJsonFormatter formatter = new SizeLimitedJsonFormatter(_inner, maxBytes: 100);
+			string result = Format(formatter, logEvent);
+
+			JsonDocument doc = JsonDocument.Parse(result.TrimEnd());
+			string messageTemplate = doc.RootElement.GetProperty("MessageTemplate").GetString()!;
+			string renderedMessage = doc.RootElement.GetProperty("RenderedMessage").GetString()!;
+
+			Assert.AreEqual("User {User} performed {Action}", messageTemplate, "MessageTemplate should contain raw template");
+			Assert.IsTrue(renderedMessage.Contains("Alice", StringComparison.Ordinal), "RenderedMessage should contain resolved User value");
+			Assert.IsTrue(renderedMessage.Contains("deploy", StringComparison.Ordinal), "RenderedMessage should contain resolved Action value");
+			Assert.IsFalse(renderedMessage.Contains("{User}", StringComparison.Ordinal), "RenderedMessage should not contain unresolved placeholder");
+			Assert.IsFalse(renderedMessage.Contains("{Action}", StringComparison.Ordinal), "RenderedMessage should not contain unresolved placeholder");
 		}
 	}
 }
diff --git a/Engine/Source/Programs/Horde/HordeServer/Utilities/SizeLimitedJsonFormatter.cs b/Engine/Source/Programs/Horde/HordeServer/Utilities/SizeLimitedJsonFormatter.cs
--- a/Engine/Source/Programs/Horde/HordeServer/Utilities/SizeLimitedJsonFormatter.cs
+++ b/Engine/Source/Programs/Horde/HordeServer/Utilities/SizeLimitedJsonFormatter.cs
@@ -3,6 +3,7 @@
 using System.IO;
 using System.Text.Encodings.Web;
 using System.Text.Json;
+using System.Text.Json.Nodes;
 using Serilog.Events;
 using Serilog.Formatting;
 
@@ -10,13 +11,20 @@ namespace HordeServer.Utilities
 {
 	/// <summary>
 	/// Wraps an <see cref="ITextFormatter"/> and enforces a maximum byte size on the formatted output.
-	/// If the inner formatter produces output exceeding the limit, a truncated replacement is emitted
-	/// containing the timestamp, level, message template, exception type/message, and a _truncated marker.
+	/// Uses three-tier truncation:
+	/// <list type="bullet">
+	///   <item>Tier 1: output fits within the limit — pass through unchanged.</item>
+	///   <item>Tier 2: strip redundant enricher properties (ExceptionInfo/ExceptionDetail) and retry.</item>
+	///   <item>Tier 3: emit minimal JSON preserving RenderedMessage and full Exception stack trace.</item>
+	/// </list>
 	/// </summary>
 	sealed class SizeLimitedJsonFormatter : ITextFormatter
 	{
 		static readonly JsonSerializerOptions s_jsonOptions = new JsonSerializerOptions { Encoder = JavaScriptEncoder.UnsafeRelaxedJsonEscaping };
 
+		const string ExceptionInfoProperty = "ExceptionInfo";
+		const string ExceptionDetailProperty = "ExceptionDetail";
+
 		readonly ITextFormatter _inner;
 		readonly int _maxBytes;
 
@@ -42,23 +50,93 @@ namespace HordeServer.Utilities
 			int byteCount = System.Text.Encoding.UTF8.GetByteCount(formatted);
 			if (byteCount <= _maxBytes)
 			{
+				// Tier 1: fits as-is
 				output.Write(formatted);
 				return;
 			}
 
-			// Emit a truncated replacement that is still valid JSON for Datadog
+			// Tier 2: try stripping redundant enricher properties
+			if (TryStripRedundantProperties(formatted, byteCount, out string? stripped))
+			{
+				output.Write(stripped);
+				return;
+			}
+
+			// Tier 3: minimal fallback with RenderedMessage and Exception
+			WriteTier3Fallback(logEvent, byteCount, output);
+		}
+
+		/// <summary>
+		/// Attempts to strip ExceptionInfo and ExceptionDetail from the Properties object.
+		/// Returns true if the stripped output fits within <see cref="_maxBytes"/>.
+		/// </summary>
+		bool TryStripRedundantProperties(string formatted, int originalByteCount, out string? result)
+		{
+			result = null;
+
+			JsonNode? root;
+			try
+			{
+				root = JsonNode.Parse(formatted);
+			}
+			catch (JsonException)
+			{
+				return false;
+			}
+
+			if (root is not JsonObject rootObj)
+			{
+				return false;
+			}
+
+			JsonNode? propertiesNode = rootObj["Properties"];
+			if (propertiesNode is not JsonObject properties)
+			{
+				return false;
+			}
+
+			bool removed = false;
+			removed |= properties.Remove(ExceptionInfoProperty);
+			removed |= properties.Remove(ExceptionDetailProperty);
+
+			if (!removed)
+			{
+				return false;
+			}
+
+			rootObj["_truncated"] = true;
+			rootObj["_originalBytes"] = originalByteCount;
+
+			string strippedJson = rootObj.ToJsonString(s_jsonOptions) + "\n";
+			int strippedByteCount = System.Text.Encoding.UTF8.GetByteCount(strippedJson);
+			if (strippedByteCount > _maxBytes)
+			{
+				return false;
+			}
+
+			result = strippedJson;
+			return true;
+		}
+
+		/// <summary>
+		/// Emits a minimal JSON object preserving RenderedMessage and the full Exception stack trace.
+		/// </summary>
+		static void WriteTier3Fallback(LogEvent logEvent, int originalByteCount, TextWriter output)
+		{
 			string timestamp = logEvent.Timestamp.UtcDateTime.ToString("O");
 			string level = logEvent.Level.ToString();
 			string messageTemplate = EscapeJson(logEvent.MessageTemplate.Text);
-			string exceptionType = logEvent.Exception?.GetType().FullName ?? "";
-			string exceptionMessage = EscapeJson(logEvent.Exception?.Message ?? "");
+			string renderedMessage = EscapeJson(logEvent.RenderMessage());
 
-			output.Write($"{{\"Timestamp\":\"{timestamp}\",\"Level\":\"{level}\",\"MessageTemplate\":\"{messageTemplate}\"");
+			output.Write($"{{\"Timestamp\":\"{timestamp}\",\"Level\":\"{level}\",\"MessageTemplate\":\"{messageTemplate}\",\"RenderedMessage\":\"{renderedMessage}\"");
 			if (logEvent.Exception != null)
 			{
-				output.Write($",\"ExceptionType\":\"{EscapeJson(exceptionType)}\",\"ExceptionMessage\":\"{exceptionMessage}\"");
+				string exceptionType = EscapeJson(logEvent.Exception.GetType().FullName ?? "");
+				string exceptionMessage = EscapeJson(logEvent.Exception.Message);
+				string exceptionFull = EscapeJson(logEvent.Exception.ToString());
+				output.Write($",\"Exception\":\"{exceptionFull}\",\"ExceptionType\":\"{exceptionType}\",\"ExceptionMessage\":\"{exceptionMessage}\"");
 			}
-			output.Write($",\"_truncated\":true,\"_originalBytes\":{byteCount}}}\n");
+			output.Write($",\"_truncated\":true,\"_originalBytes\":{originalByteCount}}}\n");
 		}
 
 		static string EscapeJson(string value)
