Fix OnJobScheduled handler leak in NotificationService.DisposeAsync

DisposeAsync was using += instead of -= for OnJobScheduled, causing
handler accumulation on each service disposal cycle. This resulted in
exponential growth of empty-pool notifications in the Redis batch queue.

diff --git a/Engine/Source/Programs/Horde/Plugins/Build/HordeServer.Build/Jobs/JobTaskSource.cs b/Engine/Source/Programs/Horde/Plugins/Build/HordeServer.Build/Jobs/JobTaskSource.cs
--- a/Engine/Source/Programs/Horde/Plugins/Build/HordeServer.Build/Jobs/JobTaskSource.cs
+++ b/Engine/Source/Programs/Horde/Plugins/Build/HordeServer.Build/Jobs/JobTaskSource.cs
@@ -495,7 +495,7 @@ namespace HordeServer.Jobs
 							newBatchIdToQueueItem[(newJob.Id, batch.Id)] = newQueueItem;
 
 							IPoolConfig? newJobPool = pools.FirstOrDefault(p => p.Id == agentType.Pool);
-							if (newJobPool != null)
+							if (newJobPool != null && !_batchIdToQueueItem.ContainsKey((newJob.Id, batch.Id)))
 							{
 								OnJobScheduled?.Invoke(newJobPool, poolStatus[agentType.Pool].HasOnlineAgents, newJob, graph, batch.Id);
 							}
diff --git a/Engine/Source/Programs/Horde/Plugins/Build/HordeServer.Build/Notifications/NotificationService.cs b/Engine/Source/Programs/Horde/Plugins/Build/HordeServer.Build/Notifications/NotificationService.cs
--- a/Engine/Source/Programs/Horde/Plugins/Build/HordeServer.Build/Notifications/NotificationService.cs
+++ b/Engine/Source/Programs/Horde/Plugins/Build/HordeServer.Build/Notifications/NotificationService.cs
@@ -189,7 +189,7 @@ namespace HordeServer.Notifications
 			_issueService.OnIssueUpdated -= NotifyIssueUpdated;
 			_jobService.OnJobAborted -= NotifyJobAborted;
 			_jobService.OnJobStepComplete -= NotifyJobStepComplete;
-			_jobService.OnJobScheduled += NotifyJobScheduled;
+			_jobService.OnJobScheduled -= NotifyJobScheduled;
 			_jobService.OnLabelUpdate -= NotifyLabelUpdate;
 
 			await _ticker.DisposeAsync();
