Horde: Fix empty pool notification spam, Only fire OnJobScheduled once per batch instead of every 5s queue refresh, preventing 73K+ entries accumulating in Redis batch queue

diff --git a/Engine/Source/Programs/Horde/Plugins/Build/HordeServer.Build/Jobs/JobTaskSource.cs b/Engine/Source/Programs/Horde/Plugins/Build/HordeServer.Build/Jobs/JobTaskSource.cs
--- a/Engine/Source/Programs/Horde/Plugins/Build/HordeServer.Build/Jobs/JobTaskSource.cs
+++ b/Engine/Source/Programs/Horde/Plugins/Build/HordeServer.Build/Jobs/JobTaskSource.cs
@@ -218,9 +218,12 @@ namespace HordeServer.Jobs
 		// Map from batch id to the corresponding queue item
 		Dictionary<(JobId, JobStepBatchId), QueueItem> _batchIdToQueueItem = new Dictionary<(JobId, JobStepBatchId), QueueItem>();
 
-		// Set of long-poll tasks waiting to be satisfied 
+		// Set of long-poll tasks waiting to be satisfied
 		readonly HashSet<QueueWaiter> _waiters = new HashSet<QueueWaiter>();
 
+		// Batches for which OnJobScheduled has already been fired (prevents re-firing on every queue refresh)
+		readonly HashSet<(JobId, JobStepBatchId)> _scheduledNotificationFired = new HashSet<(JobId, JobStepBatchId)>();
+
 		// During a background queue refresh operation, any updated batches are added to this dictionary for merging into the updated queue.
 		List<QueueItem>? _newQueueItemsDuringUpdate;
 
@@ -495,7 +498,7 @@ namespace HordeServer.Jobs
 							newBatchIdToQueueItem[(newJob.Id, batch.Id)] = newQueueItem;
 
 							IPoolConfig? newJobPool = pools.FirstOrDefault(p => p.Id == agentType.Pool);
-							if (newJobPool != null && !_batchIdToQueueItem.ContainsKey((newJob.Id, batch.Id)))
+							if (newJobPool != null && _scheduledNotificationFired.Add((newJob.Id, batch.Id)))
 							{
 								OnJobScheduled?.Invoke(newJobPool, poolStatus[agentType.Pool].HasOnlineAgents, newJob, graph, batch.Id);
 							}
@@ -540,6 +543,9 @@ namespace HordeServer.Jobs
 
 				// Clear out the list to capture queue items during an update
 				_newQueueItemsDuringUpdate = null;
+
+				// Prune notification tracking to only batches still in the queue
+				_scheduledNotificationFired.IntersectWith(_batchIdToQueueItem.Keys);
 			}
 		}
 
