Horde: Detect edge server change in PerforceMaterializer to force clean re-sync

Add client name comparison to the isDirty check in PerforceMaterializer so that switching between Perforce edge servers (which changes the client name) triggers a full workspace nuke and re-sync instead of an incorrect incremental sync against a stale have-list.

diff --git a/Engine/Source/Programs/Horde/Drivers/JobDriver.Tests/PerforceMaterializerTests.cs b/Engine/Source/Programs/Horde/Drivers/JobDriver.Tests/PerforceMaterializerTests.cs
--- a/Engine/Source/Programs/Horde/Drivers/JobDriver.Tests/PerforceMaterializerTests.cs
+++ b/Engine/Source/Programs/Horde/Drivers/JobDriver.Tests/PerforceMaterializerTests.cs
@@ -239,6 +239,37 @@ public class PerforceMaterializerTests : BasePerforceFixtureTest
 		PerforceSnapshot.GetSnapshot().GetStream(PerforceSnapshot.FooMain).GetWorkspace(9).AssertDepotFiles(pm.SyncDir.FullName);
 	}
 	
+	[TestMethod]
+	public async Task Sync_WithChangedClientName_ShouldForceCleanResync_Async()
+	{
+		CancellationToken ct = CancellationToken.None;
+		using PerforceMaterializer pm = CreatePm(PerforceSnapshot.FooMain);
+
+		// Step 1: Sync to CL 2 and finalize so state is saved to disk
+		await pm.SyncAsync(2, -1, new SyncOptions { RemoveUntracked = false }, ct);
+		PerforceSnapshot.GetSnapshot().GetStream(PerforceSnapshot.FooMain).GetWorkspace(2).AssertDepotFiles(pm.SyncDir.FullName);
+		await pm.FinalizeAsync(ct);
+
+		// Step 2: Write a marker file into the workspace directory
+		string markerPath = Path.Join(pm.SyncDir.FullName, "_edge_test_marker.txt");
+		await File.WriteAllTextAsync(markerPath, "marker", ct);
+		Assert.IsTrue(File.Exists(markerPath), "Marker file should exist before re-sync");
+
+		// Step 3: Tamper with the client name in State.json to simulate an edge server switch
+		PerforceMaterializer.State? state = await pm.LoadStateForTestAsync(ct);
+		Assert.IsNotNull(state);
+		state = state with { Client = "Horde+PM+FAKE+id+fakeEdge" };
+		string stateFilePath = Path.Join(pm.BaseDir.FullName, "State.json");
+		await File.WriteAllTextAsync(stateFilePath, JsonSerializer.Serialize(state, new JsonSerializerOptions { WriteIndented = true }), ct);
+
+		// Step 4: Sync to CL 5 — the changed client name should trigger isDirty and nuke the workspace
+		await pm.SyncAsync(5, -1, new SyncOptions { RemoveUntracked = false }, ct);
+		PerforceSnapshot.GetSnapshot().GetStream(PerforceSnapshot.FooMain).GetWorkspace(5).AssertDepotFiles(pm.SyncDir.FullName);
+
+		// Step 5: Assert the marker file was deleted (workspace was nuked because isDirty was true)
+		Assert.IsFalse(File.Exists(markerPath), "Marker file should have been deleted by workspace nuke triggered by client name change");
+	}
+
 	[TestMethod]
 	[DataRow("Platforms/...", "//Foo/Main/Platforms/... //MyComputer/Platforms/...")]
 	[DataRow(" Platforms/...", "//Foo/Main/Platforms/... //MyComputer/Platforms/...")]
diff --git a/Engine/Source/Programs/Horde/Drivers/JobDriver/Execution/PerforceMaterializer.cs b/Engine/Source/Programs/Horde/Drivers/JobDriver/Execution/PerforceMaterializer.cs
--- a/Engine/Source/Programs/Horde/Drivers/JobDriver/Execution/PerforceMaterializer.cs
+++ b/Engine/Source/Programs/Horde/Drivers/JobDriver/Execution/PerforceMaterializer.cs
@@ -286,6 +286,7 @@ public sealed class PerforceMaterializer : IWorkspaceMaterializer
 			|| state.Status == TransactionStatus.Dirty
 			|| state.Identifier != Identifier
 			|| state.Stream != _options.AgentWorkspace.Stream
+			|| state.Client != _perforceSettings!.ClientName
 			|| state.Changelist == -1;
 		
 		_logger.LogInformation("Status: {Status}\nVersion: {PrevVersion} -> {Version}\nClient: {PrevClient} -> {Client}\nIdentifier: {PrevId} -> {Id}\nStream: {PrevStream} -> {Stream}\nCL: {PrevCl} -> {Cl}\nShelve: {PrevShelve} -> {Shelve}\nDirty: {Dirty}",
