Horde: Updating schema documentation

#rnx

diff --git a/Engine/Source/Programs/Horde/Docs/Config/Schema/AclActions.md b/Engine/Source/Programs/Horde/Docs/Config/Schema/AclActions.md
--- a/Engine/Source/Programs/Horde/Docs/Config/Schema/AclActions.md
+++ b/Engine/Source/Programs/Horde/Docs/Config/Schema/AclActions.md
@@ -63,44 +63,6 @@
 | `DeleteAccount` |  |
 | `ViewAccount` |  |
 
-## Agents
-
-| Name | Description |
-| ---- | ----------- |
-| `CreateAgent` |  |
-| `CreateWorkstationAgent` |  |
-| `UpdateAgent` |  |
-| `DeleteAgent` |  |
-| `ViewAgent` |  |
-| `ListAgents` |  |
-
-## Compute
-
-| Name | Description |
-| ---- | ----------- |
-| `AddComputeTasks` |  |
-| `GetComputeTasks` |  |
-| `UbaCacheRead` |  |
-| `UbaCacheWrite` |  |
-
-## Leases
-
-| Name | Description |
-| ---- | ----------- |
-| `ViewLeases` |  |
-| `ViewLeaseTasks` |  |
-
-## Logs
-
-| Name | Description |
-| ---- | ----------- |
-| `CreateLog` |  |
-| `UpdateLog` |  |
-| `ViewLog` |  |
-| `WriteLogData` |  |
-| `CreateEvent` |  |
-| `ViewEvent` |  |
-
 ## Notices
 
 | Name | Description |
@@ -109,16 +71,6 @@
 | `UpdateNotice` |  |
 | `DeleteNotice` |  |
 
-## Pools
-
-| Name | Description |
-| ---- | ----------- |
-| `CreatePool` |  |
-| `UpdatePool` |  |
-| `DeletePool` |  |
-| `ViewPool` |  |
-| `ListPools` |  |
-
 ## ServiceAccounts
 
 | Name | Description |
@@ -128,21 +80,6 @@
 | `DeleteAccount` |  |
 | `ViewAccount` |  |
 
-## Sessions
-
-| Name | Description |
-| ---- | ----------- |
-| `CreateSession` |  |
-| `ViewSession` |  |
-
-## Software
-
-| Name | Description |
-| ---- | ----------- |
-| `UploadSoftware` |  |
-| `DownloadSoftware` |  |
-| `DeleteSoftware` |  |
-
 ## Accounts
 
 | Name | Description |
@@ -163,30 +100,6 @@
 | `ViewAgent` |  |
 | `ListAgents` |  |
 
-## Artifacts
-
-| Name | Description |
-| ---- | ----------- |
-| `ReadArtifact` |  |
-| `WriteArtifact` |  |
-| `DeleteArtifact` |  |
-| `UploadArtifact` |  |
-| `DownloadArtifact` |  |
-
-## Bisect
-
-| Name | Description |
-| ---- | ----------- |
-| `CreateBisectTask` |  |
-| `UpdateBisectTask` |  |
-| `ViewBisectTask` |  |
-
-## BuildHealth
-
-| Name | Description |
-| ---- | ----------- |
-| `AddUpdateDeleteBuildHealthFilter` |  |
-
 ## Compute
 
 | Name | Description |
@@ -196,24 +109,6 @@
 | `UbaCacheRead` |  |
 | `UbaCacheWrite` |  |
 
-## Devices
-
-| Name | Description |
-| ---- | ----------- |
-| `DeviceRead` |  |
-| `DeviceWrite` |  |
-
-## Jobs
-
-| Name | Description |
-| ---- | ----------- |
-| `CreateJob` |  |
-| `UpdateJob` |  |
-| `DeleteJob` |  |
-| `ExecuteJob` |  |
-| `RetryJobStep` |  |
-| `ViewJob` |  |
-
 ## Leases
 
 | Name | Description |
@@ -240,12 +135,6 @@
 | `UpdateNotice` |  |
 | `DeleteNotice` |  |
 
-## Notifications
-
-| Name | Description |
-| ---- | ----------- |
-| `CreateSubscription` |  |
-
 ## Pools
 
 | Name | Description |
@@ -256,22 +145,6 @@
 | `ViewPool` |  |
 | `ListPools` |  |
 
-## Projects
-
-| Name | Description |
-| ---- | ----------- |
-| `CreateProject` |  |
-| `DeleteProject` |  |
-| `UpdateProject` |  |
-| `ViewProject` |  |
-
-## Replicators
-
-| Name | Description |
-| ---- | ----------- |
-| `UpdateReplicator` |  |
-| `ViewReplicator` |  |
-
 ## ServiceAccounts
 
 | Name | Description |
@@ -296,24 +169,6 @@
 | `DownloadSoftware` |  |
 | `DeleteSoftware` |  |
 
-## Streams
-
-| Name | Description |
-| ---- | ----------- |
-| `CreateStream` |  |
-| `UpdateStream` |  |
-| `DeleteStream` |  |
-| `ViewStream` |  |
-| `ViewChanges` |  |
-| `ViewTemplate` |  |
-
-## TestData
-
-| Name | Description |
-| ---- | ----------- |
-| `TestHealthRead` |  |
-| `TestHealthWrite` |  |
-
 ## Accounts
 
 | Name | Description |
@@ -638,6 +493,19 @@
 | `DownloadSoftware` |  |
 | `DeleteSoftware` |  |
 
+## Storage
+
+| Name | Description |
+| ---- | ----------- |
+| `ReadBlobs` |  |
+| `WriteBlobs` |  |
+| `ReadAliases` |  |
+| `WriteAliases` |  |
+| `DeleteAliases` |  |
+| `ReadRefs` |  |
+| `WriteRefs` |  |
+| `DeleteRefs` |  |
+
 ## Streams
 
 | Name | Description |
@@ -785,12 +653,6 @@
 | `UpdateReplicator` |  |
 | `ViewReplicator` |  |
 
-## Secrets
-
-| Name | Description |
-| ---- | ----------- |
-| `ViewSecret` |  |
-
 ## ServiceAccounts
 
 | Name | Description |
@@ -815,6 +677,19 @@
 | `DownloadSoftware` |  |
 | `DeleteSoftware` |  |
 
+## Storage
+
+| Name | Description |
+| ---- | ----------- |
+| `ReadBlobs` |  |
+| `WriteBlobs` |  |
+| `ReadAliases` |  |
+| `WriteAliases` |  |
+| `DeleteAliases` |  |
+| `ReadRefs` |  |
+| `WriteRefs` |  |
+| `DeleteRefs` |  |
+
 ## Streams
 
 | Name | Description |
@@ -833,6 +708,13 @@
 | `TestHealthRead` |  |
 | `TestHealthWrite` |  |
 
+## Tools
+
+| Name | Description |
+| ---- | ----------- |
+| `DownloadTool` |  |
+| `UploadTool` |  |
+
 ## Accounts
 
 | Name | Description |
@@ -962,12 +844,6 @@
 | `UpdateReplicator` |  |
 | `ViewReplicator` |  |
 
-## Secrets
-
-| Name | Description |
-| ---- | ----------- |
-| `ViewSecret` |  |
-
 ## ServiceAccounts
 
 | Name | Description |
@@ -1023,6 +899,13 @@
 | `TestHealthRead` |  |
 | `TestHealthWrite` |  |
 
+## Tools
+
+| Name | Description |
+| ---- | ----------- |
+| `DownloadTool` |  |
+| `UploadTool` |  |
+
 ## Accounts
 
 | Name | Description |
@@ -1206,18 +1089,19 @@
 | `ViewChanges` |  |
 | `ViewTemplate` |  |
 
-## Symbols
+## TestData
 
 | Name | Description |
 | ---- | ----------- |
-| `ReadSymbols` |  |
+| `TestHealthRead` |  |
+| `TestHealthWrite` |  |
 
-## TestData
+## Tools
 
 | Name | Description |
 | ---- | ----------- |
-| `TestHealthRead` |  |
-| `TestHealthWrite` |  |
+| `DownloadTool` |  |
+| `UploadTool` |  |
 
 ## Accounts
 
diff --git a/Engine/Source/Programs/Horde/Docs/Deployment/ServerSettings.md b/Engine/Source/Programs/Horde/Docs/Deployment/ServerSettings.md
--- a/Engine/Source/Programs/Horde/Docs/Deployment/ServerSettings.md
+++ b/Engine/Source/Programs/Horde/Docs/Deployment/ServerSettings.md
@@ -126,14 +126,14 @@
 Name | Description
 ---- | -----------
 `analytics` | [AnalyticsServerConfig](#analyticsserverconfig)<br>Configuration for the analytics plugin
+`exampleplugin` | [ExamplePluginServerConfig](#examplepluginserverconfig)<br>Configuration for the exampleplugin plugin
 `compute` | [ComputeServerConfig](#computeserverconfig)<br>Configuration for the compute plugin
 `build` | [BuildServerConfig](#buildserverconfig)<br>Configuration for the build plugin
+`storage` | [StorageServerConfig](#storageserverconfig)<br>Configuration for the storage plugin
+`tools` | [ToolsServerConfig](#toolsserverconfig)<br>Configuration for the tools plugin
 `experimental` | [ExperimentalServerConfig](#experimentalserverconfig)<br>Configuration for the experimental plugin
-`exampleplugin` | [ExamplePluginServerConfig](#examplepluginserverconfig)<br>Configuration for the exampleplugin plugin
 `secrets` | [SecretsServerConfig](#secretsserverconfig)<br>Configuration for the secrets plugin
-`storage` | [StorageServerConfig](#storageserverconfig)<br>Configuration for the storage plugin
 `symbols` | [PluginServerConfig](#pluginserverconfig)<br>Configuration for the symbols plugin
-`tools` | [ToolsServerConfig](#toolsserverconfig)<br>Configuration for the tools plugin
 
 ## AnalyticsServerConfig
 
@@ -164,6 +164,12 @@
 `retainDays` | `number`<br>
 `enabled` | `boolean`<br>
 
+## ExamplePluginServerConfig
+
+Name | Description
+---- | -----------
+`enabled` | `boolean`<br>
+
 ## ComputeServerConfig
 
 Name | Description
@@ -342,25 +348,6 @@
 `maxChildCount` | `integer`<br>Maximum number of children in each node
 `sliceThreshold` | `integer`<br>Threshold hash value for splitting interior nodes
 
-## ExperimentalServerConfig
-
-Name | Description
----- | -----------
-`enabled` | `boolean`<br>
-
-## ExamplePluginServerConfig
-
-Name | Description
----- | -----------
-`enabled` | `boolean`<br>
-
-## SecretsServerConfig
-
-Name | Description
----- | -----------
-`withAws` | `boolean`<br>
-`enabled` | `boolean`<br>
-
 ## StorageServerConfig
 
 Name | Description
@@ -412,12 +399,6 @@
 `AssumeRole` | 
 `AssumeRoleWebIdentity` | 
 
-## PluginServerConfig
-
-Name | Description
----- | -----------
-`enabled` | `boolean`<br>
-
 ## ToolsServerConfig
 
 Name | Description
@@ -478,3 +459,22 @@
 `actions` | `string[]`<br>
 `excludeActions` | `string[]`<br>
 `extends` | `string[]`<br>
+
+## ExperimentalServerConfig
+
+Name | Description
+---- | -----------
+`enabled` | `boolean`<br>
+
+## SecretsServerConfig
+
+Name | Description
+---- | -----------
+`withAws` | `boolean`<br>
+`enabled` | `boolean`<br>
+
+## PluginServerConfig
+
+Name | Description
+---- | -----------
+`enabled` | `boolean`<br>
