Horde: Fix getMore failures in AgentCollection.FindAsync

Replace cursor-based iteration with ToListAsync() in AgentCollection.FindAsync to eliminate long-lived MongoDB cursors vulnerable to getMore failures on secondary reads.

diff --git a/Engine/Source/Programs/Horde/Plugins/Compute/HordeServer.Compute/Agents/AgentCollection.cs b/Engine/Source/Programs/Horde/Plugins/Compute/HordeServer.Compute/Agents/AgentCollection.cs
--- a/Engine/Source/Programs/Horde/Plugins/Compute/HordeServer.Compute/Agents/AgentCollection.cs
+++ b/Engine/Source/Programs/Horde/Plugins/Compute/HordeServer.Compute/Agents/AgentCollection.cs
@@ -711,23 +711,19 @@ namespace HordeServer.Agents
 
 			IMongoCollection<AgentDocument> collection = consistentRead ? _agentCollection : _agentCollection.WithReadPreference(ReadPreference.SecondaryPreferred);
 
-			using (IAsyncCursor<AgentDocument> cursor = await collection.Find(filter).ToCursorAsync(cancellationToken))
+			// Materialize all documents up front to avoid long-lived cursors vulnerable to getMore failures on secondary reads.
+			// Post-load processing (document migration, Redis session lookups) can take significant time between cursor batches.
+			List<AgentDocument> allDocuments = await collection.Find(filter).ToListAsync(cancellationToken);
+			allDocuments = await PostLoadAsync(allDocuments, cancellationToken);
+			foreach (IAgent agent in await CreateAgentObjectsBatchedAsync(allDocuments, cancellationToken))
 			{
-				while (await cursor.MoveNextAsync(cancellationToken))
+				if (status == null || agent.Status == status.Value)
 				{
-					List<AgentDocument> documents = await PostLoadAsync(cursor.Current.ToList(), cancellationToken);
-					foreach (IAgent agent in await CreateAgentObjectsBatchedAsync(documents, cancellationToken))
-					{
-						if (status == null || agent.Status == status.Value)
-						{
-							yield return agent;
-						}
-					}
+					yield return agent;
 				}
 			}
 		}
 
-		/// <inheritdoc/>
 		async Task<SessionDocument?> GetSessionAsync(AgentId agentId, SessionId sessionId, CancellationToken cancellationToken = default)
 		{
 			FilterDefinitionBuilder<SessionDocument> filterBuilder = Builders<SessionDocument>.Filter;
@@ -736,7 +732,6 @@ namespace HordeServer.Agents
 			return await _sessionCollection.Find(filter).FirstOrDefaultAsync(cancellationToken);
 		}
 
-		/// <inheritdoc/>
 		async Task<List<SessionDocument>> FindSessionsAsync(AgentId agentId, DateTime? startTime, DateTime? finishTime, int index, int count, CancellationToken cancellationToken = default)
 		{
 			FilterDefinitionBuilder<SessionDocument> filterBuilder = Builders<SessionDocument>.Filter;
