Horde: Updating schema documentation

#rnx

diff --git a/Engine/Source/Programs/Horde/Docs/Config/Schema/AclActions.md b/Engine/Source/Programs/Horde/Docs/Config/Schema/AclActions.md
--- a/Engine/Source/Programs/Horde/Docs/Config/Schema/AclActions.md
+++ b/Engine/Source/Programs/Horde/Docs/Config/Schema/AclActions.md
@@ -89,43 +89,47 @@
 | `DeleteAccount` |  |
 | `ViewAccount` |  |
 
-## Agents
+## Artifacts
 
 | Name | Description |
 | ---- | ----------- |
-| `CreateAgent` |  |
-| `CreateWorkstationAgent` |  |
-| `UpdateAgent` |  |
-| `DeleteAgent` |  |
-| `ViewAgent` |  |
-| `ListAgents` |  |
+| `ReadArtifact` |  |
+| `WriteArtifact` |  |
+| `DeleteArtifact` |  |
+| `UploadArtifact` |  |
+| `DownloadArtifact` |  |
 
-## Compute
+## Bisect
 
 | Name | Description |
 | ---- | ----------- |
-| `AddComputeTasks` |  |
-| `GetComputeTasks` |  |
-| `UbaCacheRead` |  |
-| `UbaCacheWrite` |  |
+| `CreateBisectTask` |  |
+| `UpdateBisectTask` |  |
+| `ViewBisectTask` |  |
 
-## Leases
+## BuildHealth
 
 | Name | Description |
 | ---- | ----------- |
-| `ViewLeases` |  |
-| `ViewLeaseTasks` |  |
+| `AddUpdateDeleteBuildHealthFilter` |  |
 
-## Logs
+## Devices
 
 | Name | Description |
 | ---- | ----------- |
-| `CreateLog` |  |
-| `UpdateLog` |  |
-| `ViewLog` |  |
-| `WriteLogData` |  |
-| `CreateEvent` |  |
-| `ViewEvent` |  |
+| `DeviceRead` |  |
+| `DeviceWrite` |  |
+
+## Jobs
+
+| Name | Description |
+| ---- | ----------- |
+| `CreateJob` |  |
+| `UpdateJob` |  |
+| `DeleteJob` |  |
+| `ExecuteJob` |  |
+| `RetryJobStep` |  |
+| `ViewJob` |  |
 
 ## Notices
 
@@ -135,15 +139,27 @@
 | `UpdateNotice` |  |
 | `DeleteNotice` |  |
 
-## Pools
+## Notifications
 
 | Name | Description |
 | ---- | ----------- |
-| `CreatePool` |  |
-| `UpdatePool` |  |
-| `DeletePool` |  |
-| `ViewPool` |  |
-| `ListPools` |  |
+| `CreateSubscription` |  |
+
+## Projects
+
+| Name | Description |
+| ---- | ----------- |
+| `CreateProject` |  |
+| `DeleteProject` |  |
+| `UpdateProject` |  |
+| `ViewProject` |  |
+
+## Replicators
+
+| Name | Description |
+| ---- | ----------- |
+| `UpdateReplicator` |  |
+| `ViewReplicator` |  |
 
 ## ServiceAccounts
 
@@ -154,20 +170,23 @@
 | `DeleteAccount` |  |
 | `ViewAccount` |  |
 
-## Sessions
+## Streams
 
 | Name | Description |
 | ---- | ----------- |
-| `CreateSession` |  |
-| `ViewSession` |  |
+| `CreateStream` |  |
+| `UpdateStream` |  |
+| `DeleteStream` |  |
+| `ViewStream` |  |
+| `ViewChanges` |  |
+| `ViewTemplate` |  |
 
-## Software
+## TestData
 
 | Name | Description |
 | ---- | ----------- |
-| `UploadSoftware` |  |
-| `DownloadSoftware` |  |
-| `DeleteSoftware` |  |
+| `TestHealthRead` |  |
+| `TestHealthWrite` |  |
 
 ## Accounts
 
@@ -178,17 +197,6 @@
 | `DeleteAccount` |  |
 | `ViewAccount` |  |
 
-## Agents
-
-| Name | Description |
-| ---- | ----------- |
-| `CreateAgent` |  |
-| `CreateWorkstationAgent` |  |
-| `UpdateAgent` |  |
-| `DeleteAgent` |  |
-| `ViewAgent` |  |
-| `ListAgents` |  |
-
 ## Artifacts
 
 | Name | Description |
@@ -213,15 +221,6 @@
 | ---- | ----------- |
 | `AddUpdateDeleteBuildHealthFilter` |  |
 
-## Compute
-
-| Name | Description |
-| ---- | ----------- |
-| `AddComputeTasks` |  |
-| `GetComputeTasks` |  |
-| `UbaCacheRead` |  |
-| `UbaCacheWrite` |  |
-
 ## Devices
 
 | Name | Description |
@@ -240,24 +239,6 @@
 | `RetryJobStep` |  |
 | `ViewJob` |  |
 
-## Leases
-
-| Name | Description |
-| ---- | ----------- |
-| `ViewLeases` |  |
-| `ViewLeaseTasks` |  |
-
-## Logs
-
-| Name | Description |
-| ---- | ----------- |
-| `CreateLog` |  |
-| `UpdateLog` |  |
-| `ViewLog` |  |
-| `WriteLogData` |  |
-| `CreateEvent` |  |
-| `ViewEvent` |  |
-
 ## Notices
 
 | Name | Description |
@@ -272,16 +253,6 @@
 | ---- | ----------- |
 | `CreateSubscription` |  |
 
-## Pools
-
-| Name | Description |
-| ---- | ----------- |
-| `CreatePool` |  |
-| `UpdatePool` |  |
-| `DeletePool` |  |
-| `ViewPool` |  |
-| `ListPools` |  |
-
 ## Projects
 
 | Name | Description |
@@ -298,6 +269,12 @@
 | `UpdateReplicator` |  |
 | `ViewReplicator` |  |
 
+## Secrets
+
+| Name | Description |
+| ---- | ----------- |
+| `ViewSecret` |  |
+
 ## ServiceAccounts
 
 | Name | Description |
@@ -307,21 +284,6 @@
 | `DeleteAccount` |  |
 | `ViewAccount` |  |
 
-## Sessions
-
-| Name | Description |
-| ---- | ----------- |
-| `CreateSession` |  |
-| `ViewSession` |  |
-
-## Software
-
-| Name | Description |
-| ---- | ----------- |
-| `UploadSoftware` |  |
-| `DownloadSoftware` |  |
-| `DeleteSoftware` |  |
-
 ## Streams
 
 | Name | Description |
@@ -349,17 +311,6 @@
 | `DeleteAccount` |  |
 | `ViewAccount` |  |
 
-## Agents
-
-| Name | Description |
-| ---- | ----------- |
-| `CreateAgent` |  |
-| `CreateWorkstationAgent` |  |
-| `UpdateAgent` |  |
-| `DeleteAgent` |  |
-| `ViewAgent` |  |
-| `ListAgents` |  |
-
 ## Artifacts
 
 | Name | Description |
@@ -384,15 +335,6 @@
 | ---- | ----------- |
 | `AddUpdateDeleteBuildHealthFilter` |  |
 
-## Compute
-
-| Name | Description |
-| ---- | ----------- |
-| `AddComputeTasks` |  |
-| `GetComputeTasks` |  |
-| `UbaCacheRead` |  |
-| `UbaCacheWrite` |  |
-
 ## Devices
 
 | Name | Description |
@@ -411,24 +353,6 @@
 | `RetryJobStep` |  |
 | `ViewJob` |  |
 
-## Leases
-
-| Name | Description |
-| ---- | ----------- |
-| `ViewLeases` |  |
-| `ViewLeaseTasks` |  |
-
-## Logs
-
-| Name | Description |
-| ---- | ----------- |
-| `CreateLog` |  |
-| `UpdateLog` |  |
-| `ViewLog` |  |
-| `WriteLogData` |  |
-| `CreateEvent` |  |
-| `ViewEvent` |  |
-
 ## Notices
 
 | Name | Description |
@@ -443,16 +367,6 @@
 | ---- | ----------- |
 | `CreateSubscription` |  |
 
-## Pools
-
-| Name | Description |
-| ---- | ----------- |
-| `CreatePool` |  |
-| `UpdatePool` |  |
-| `DeletePool` |  |
-| `ViewPool` |  |
-| `ListPools` |  |
-
 ## Projects
 
 | Name | Description |
@@ -469,6 +383,12 @@
 | `UpdateReplicator` |  |
 | `ViewReplicator` |  |
 
+## Secrets
+
+| Name | Description |
+| ---- | ----------- |
+| `ViewSecret` |  |
+
 ## ServiceAccounts
 
 | Name | Description |
@@ -478,34 +398,6 @@
 | `DeleteAccount` |  |
 | `ViewAccount` |  |
 
-## Sessions
-
-| Name | Description |
-| ---- | ----------- |
-| `CreateSession` |  |
-| `ViewSession` |  |
-
-## Software
-
-| Name | Description |
-| ---- | ----------- |
-| `UploadSoftware` |  |
-| `DownloadSoftware` |  |
-| `DeleteSoftware` |  |
-
-## Storage
-
-| Name | Description |
-| ---- | ----------- |
-| `ReadBlobs` |  |
-| `WriteBlobs` |  |
-| `ReadAliases` |  |
-| `WriteAliases` |  |
-| `DeleteAliases` |  |
-| `ReadRefs` |  |
-| `WriteRefs` |  |
-| `DeleteRefs` |  |
-
 ## Streams
 
 | Name | Description |
@@ -653,6 +545,12 @@
 | `UpdateReplicator` |  |
 | `ViewReplicator` |  |
 
+## Secrets
+
+| Name | Description |
+| ---- | ----------- |
+| `ViewSecret` |  |
+
 ## ServiceAccounts
 
 | Name | Description |
@@ -677,19 +575,6 @@
 | `DownloadSoftware` |  |
 | `DeleteSoftware` |  |
 
-## Storage
-
-| Name | Description |
-| ---- | ----------- |
-| `ReadBlobs` |  |
-| `WriteBlobs` |  |
-| `ReadAliases` |  |
-| `WriteAliases` |  |
-| `DeleteAliases` |  |
-| `ReadRefs` |  |
-| `WriteRefs` |  |
-| `DeleteRefs` |  |
-
 ## Streams
 
 | Name | Description |
@@ -708,13 +593,6 @@
 | `TestHealthRead` |  |
 | `TestHealthWrite` |  |
 
-## Tools
-
-| Name | Description |
-| ---- | ----------- |
-| `DownloadTool` |  |
-| `UploadTool` |  |
-
 ## Accounts
 
 | Name | Description |
@@ -844,6 +722,12 @@
 | `UpdateReplicator` |  |
 | `ViewReplicator` |  |
 
+## Secrets
+
+| Name | Description |
+| ---- | ----------- |
+| `ViewSecret` |  |
+
 ## ServiceAccounts
 
 | Name | Description |
@@ -868,19 +752,6 @@
 | `DownloadSoftware` |  |
 | `DeleteSoftware` |  |
 
-## Storage
-
-| Name | Description |
-| ---- | ----------- |
-| `ReadBlobs` |  |
-| `WriteBlobs` |  |
-| `ReadAliases` |  |
-| `WriteAliases` |  |
-| `DeleteAliases` |  |
-| `ReadRefs` |  |
-| `WriteRefs` |  |
-| `DeleteRefs` |  |
-
 ## Streams
 
 | Name | Description |
@@ -1065,19 +936,6 @@
 | `DownloadSoftware` |  |
 | `DeleteSoftware` |  |
 
-## Storage
-
-| Name | Description |
-| ---- | ----------- |
-| `ReadBlobs` |  |
-| `WriteBlobs` |  |
-| `ReadAliases` |  |
-| `WriteAliases` |  |
-| `DeleteAliases` |  |
-| `ReadRefs` |  |
-| `WriteRefs` |  |
-| `DeleteRefs` |  |
-
 ## Streams
 
 | Name | Description |
@@ -1089,6 +947,12 @@
 | `ViewChanges` |  |
 | `ViewTemplate` |  |
 
+## Symbols
+
+| Name | Description |
+| ---- | ----------- |
+| `ReadSymbols` |  |
+
 ## TestData
 
 | Name | Description |
diff --git a/Engine/Source/Programs/Horde/Docs/Config/Schema/Globals.md b/Engine/Source/Programs/Horde/Docs/Config/Schema/Globals.md
--- a/Engine/Source/Programs/Horde/Docs/Config/Schema/Globals.md
+++ b/Engine/Source/Programs/Horde/Docs/Config/Schema/Globals.md
@@ -57,13 +57,13 @@
 ---- | -----------
 `analytics` | [AnalyticsConfig](#analyticsconfig)<br>Configuration for the analytics plugin
 `exampleplugin` | [ExamplePluginConfig](#examplepluginconfig)<br>Configuration for the exampleplugin plugin
-`compute` | [ComputeConfig](#computeconfig)<br>Configuration for the compute plugin
 `build` | [BuildConfig](#buildconfig)<br>Configuration for the build plugin
-`storage` | [StorageConfig](#storageconfig)<br>Configuration for the storage plugin
-`tools` | [ToolsConfig](#toolsconfig)<br>Configuration for the tools plugin
-`experimental` | [ExperimentalConfig](#experimentalconfig)<br>Configuration for the experimental plugin
 `secrets` | [SecretsConfig](#secretsconfig)<br>Configuration for the secrets plugin
+`experimental` | [ExperimentalConfig](#experimentalconfig)<br>Configuration for the experimental plugin
+`compute` | [ComputeConfig](#computeconfig)<br>Configuration for the compute plugin
+`tools` | [ToolsConfig](#toolsconfig)<br>Configuration for the tools plugin
 `symbols` | [SymbolsConfig](#symbolsconfig)<br>Configuration for the symbols plugin
+`storage` | [StorageConfig](#storageconfig)<br>Configuration for the storage plugin
 
 ## AnalyticsConfig
 
@@ -77,20 +77,121 @@
 ---- | -----------
 `greetingMessage` | `string`<br>
 
-## ComputeConfig
+## BuildConfig
 
 Name | Description
 ---- | -----------
-`include` | [ConfigInclude](#configinclude)`[]`<br>
-`macros` | [ConfigMacro](#configmacro)`[]`<br>
+`perforceClusters` | [PerforceCluster](#perforcecluster)`[]`<br>
+`devices` | [DeviceConfig](#deviceconfig)<br>
+`maxConformCount` | `integer`<br>
+`agentShutdownIfDisabledGracePeriod` | `string`<br>
+`artifactTypes` | [ArtifactTypeConfig](#artifacttypeconfig)`[]`<br>
+`projects` | [ProjectConfig](Projects.md)`[]`<br>
+`enableConformTasks` | `boolean`<br>
+`enableJobTaskConformCheck` | `boolean`<br>
+`issueFixedTag` | `string`<br>
+`issues` | [IssueConfig](Issues.md)<br>
+`enableIssueRootCauseAnalysisNotifications` | `boolean`<br>
+`autoRcaUseMultiStage` | `boolean`<br>
+`maxConcurrentRcaRequests` | `integer`<br>
+`issueFindSuspectsMaxSuccessDays` | `integer`<br>
+`issueFindSuspectsMaxScoringSeconds` | `integer`<br>
+`issueFindSuspectsMaxChanges` | `integer`<br>
+`issueFindSuspectsEnableParallelScoring` | `boolean`<br>
+`issueFindSuspectsMaxParallelism` | `integer`<br>
+`issueFindSuspectsEnableEarlyExit` | `boolean`<br>
+`issueFindSuspectsEnableBidirectionalQuery` | `boolean`<br>
+`issueFindSuspectsEnableFingerprintFallback` | `boolean`<br>
+`issueFingerprintFallbackMinScore` | `integer`<br>
+`issueFingerprintFallbackMaxAgeDays` | `integer`<br>
+`issueFingerprintFallbackExcludedTypes` | `string[]`<br>
+`enableExpandedIgnorePatterns` | `boolean`<br>
+
+## PerforceCluster
+
+Name | Description
+---- | -----------
+`name` | `string`<br>
+`serviceAccount` | `string`<br>
+`canImpersonate` | `boolean`<br>
+`supportsPartitionedWorkspaces` | `boolean`<br>
+`servers` | [PerforceServer](#perforceserver)`[]`<br>
+`credentials` | [PerforceCredentials](#perforcecredentials)`[]`<br>
+`autoSdk` | [AutoSdkWorkspace](#autosdkworkspace)`[]`<br>
+
+## PerforceServer
+
+Name | Description
+---- | -----------
+`serverAndPort` | `string`<br>
+`healthCheck` | `boolean`<br>
+`resolveDns` | `boolean`<br>
+`maxConformCount` | `integer`<br>
+`condition` | `string`<br>
+`properties` | `string[]`<br>
+
+## PerforceCredentials
+
+Name | Description
+---- | -----------
+`userName` | `string`<br>
+`password` | `string`<br>
+`ticket` | `string`<br>
+
+## AutoSdkWorkspace
+
+Name | Description
+---- | -----------
+`name` | `string`<br>
+`properties` | `string[]`<br>
+`userName` | `string`<br>
+`stream` | `string`<br>
+
+## DeviceConfig
+
+Name | Description
+---- | -----------
+`platforms` | [DevicePlatformConfig](#deviceplatformconfig)`[]`<br>
+`pools` | [DevicePoolConfig](#devicepoolconfig)`[]`<br>
+
+## DevicePlatformConfig
+
+Name | Description
+---- | -----------
+`id` | `string`<br>
+`name` | `string`<br>
+`models` | `string[]`<br>
+`legacyNames` | `string[]`<br>
+`legacyPerfSpecHighModel` | `string`<br>
+
+## DevicePoolConfig
+
+Name | Description
+---- | -----------
+`id` | `string`<br>
+`name` | `string`<br>
+`poolType` | [DevicePoolType](#devicepooltype-enum)<br>
+`projectIds` | `string[]`<br>
+
+## DevicePoolType (Enum)
+
+The type of device pool
+
+Name | Description
+---- | -----------
+`Automation` | Available to CIS jobs
+`Shared` | Shared by users with remote checking and checkouts
+
+## ArtifactTypeConfig
+
+Name | Description
+---- | -----------
+`name` | `string`<br>
+`type` | `string`<br>
 `acl` | [AclConfig](#aclconfig)<br>
-`versionEnum` | [ConfigVersion](#configversion-enum)<br>
-`telemetryStoreId` | `string`<br>
-`rates` | [AgentRateConfig](#agentrateconfig)`[]`<br>
-`clusters` | [ComputeClusterConfig](#computeclusterconfig)`[]`<br>
-`pools` | [PoolConfig](#poolconfig)`[]`<br>
-`software` | [AgentSoftwareConfig](#agentsoftwareconfig)`[]`<br>
-`networks` | [NetworkConfig](#networkconfig)`[]`<br>
+`keepCount` | `integer`<br>
+`keepDays` | `integer`<br>
+`namespaceId` | `string`<br>
 
 ## AclConfig
 
@@ -125,6 +226,93 @@
 `excludeActions` | `string[]`<br>
 `extends` | `string[]`<br>
 
+## SecretsConfig
+
+Name | Description
+---- | -----------
+`secrets` | [SecretConfig](#secretconfig)`[]`<br>
+`providerConfigs` | [SecretProviderConfig](#secretproviderconfig)`[]`<br>
+
+## SecretConfig
+
+Name | Description
+---- | -----------
+`id` | `string`<br>
+`data` | `string` `->` `string`<br>
+`sources` | [ExternalSecretConfig](#externalsecretconfig)`[]`<br>
+`acl` | [AclConfig](#aclconfig)<br>
+
+## ExternalSecretConfig
+
+Name | Description
+---- | -----------
+`providerConfig` | `string`<br>
+`provider` | `string`<br>
+`format` | [ExternalSecretFormat](#externalsecretformat-enum)<br>
+`key` | `string`<br>
+`path` | `string`<br>
+
+## ExternalSecretFormat (Enum)
+
+Name | Description
+---- | -----------
+`Text` | 
+`Json` | 
+
+## SecretProviderConfig
+
+Name | Description
+---- | -----------
+`name` | `string`<br>
+`provider` | `string`<br>
+`hcpVault` | [HcpVaultConfig](#hcpvaultconfig)<br>
+`azureKeyVault` | [AzureKeyVaultConfig](#azurekeyvaultconfig)<br>
+
+## HcpVaultConfig
+
+Name | Description
+---- | -----------
+`preSharedKey` | `string`<br>
+`endPoint` | `string`<br>
+`awsArnRole` | `string`<br>
+`awsIamServerId` | `string`<br>
+`role` | `string`<br>
+`credentials` | [HcpVaultCredentialsType](#hcpvaultcredentialstype-enum)<br>
+
+## HcpVaultCredentialsType (Enum)
+
+Name | Description
+---- | -----------
+`PreSharedKey` | 
+`AwsAuth` | 
+
+## AzureKeyVaultConfig
+
+Name | Description
+---- | -----------
+`vaultUri` | `string`<br>
+
+## ExperimentalConfig
+
+Name | Description
+---- | -----------
+`notifications` | [NotificationConfig](Notification.md)`[]`<br>
+
+## ComputeConfig
+
+Name | Description
+---- | -----------
+`include` | [ConfigInclude](#configinclude)`[]`<br>
+`macros` | [ConfigMacro](#configmacro)`[]`<br>
+`acl` | [AclConfig](#aclconfig)<br>
+`versionEnum` | [ConfigVersion](#configversion-enum)<br>
+`telemetryStoreId` | `string`<br>
+`rates` | [AgentRateConfig](#agentrateconfig)`[]`<br>
+`clusters` | [ComputeClusterConfig](#computeclusterconfig)`[]`<br>
+`pools` | [PoolConfig](#poolconfig)`[]`<br>
+`software` | [AgentSoftwareConfig](#agentsoftwareconfig)`[]`<br>
+`networks` | [NetworkConfig](#networkconfig)`[]`<br>
+
 ## ConfigVersion (Enum)
 
 Name | Description
@@ -277,121 +465,44 @@
 `description` | `string`<br>
 `computeId` | `string`<br>
 
-## BuildConfig
+## ToolsConfig
 
 Name | Description
 ---- | -----------
-`perforceClusters` | [PerforceCluster](#perforcecluster)`[]`<br>
-`devices` | [DeviceConfig](#deviceconfig)<br>
-`maxConformCount` | `integer`<br>
-`agentShutdownIfDisabledGracePeriod` | `string`<br>
-`artifactTypes` | [ArtifactTypeConfig](#artifacttypeconfig)`[]`<br>
-`projects` | [ProjectConfig](Projects.md)`[]`<br>
-`enableConformTasks` | `boolean`<br>
-`enableJobTaskConformCheck` | `boolean`<br>
-`issueFixedTag` | `string`<br>
-`issues` | [IssueConfig](Issues.md)<br>
-`enableIssueRootCauseAnalysisNotifications` | `boolean`<br>
-`autoRcaUseMultiStage` | `boolean`<br>
-`maxConcurrentRcaRequests` | `integer`<br>
-`issueFindSuspectsMaxSuccessDays` | `integer`<br>
-`issueFindSuspectsMaxScoringSeconds` | `integer`<br>
-`issueFindSuspectsMaxChanges` | `integer`<br>
-`issueFindSuspectsEnableParallelScoring` | `boolean`<br>
-`issueFindSuspectsMaxParallelism` | `integer`<br>
-`issueFindSuspectsEnableEarlyExit` | `boolean`<br>
-`issueFindSuspectsEnableBidirectionalQuery` | `boolean`<br>
-`issueFindSuspectsEnableFingerprintFallback` | `boolean`<br>
-`issueFingerprintFallbackMinScore` | `integer`<br>
-`issueFingerprintFallbackMaxAgeDays` | `integer`<br>
-`issueFingerprintFallbackExcludedTypes` | `string[]`<br>
-`enableExpandedIgnorePatterns` | `boolean`<br>
+`tools` | [ToolConfig](#toolconfig)`[]`<br>
 
-## PerforceCluster
-
-Name | Description
----- | -----------
-`name` | `string`<br>
-`serviceAccount` | `string`<br>
-`canImpersonate` | `boolean`<br>
-`supportsPartitionedWorkspaces` | `boolean`<br>
-`servers` | [PerforceServer](#perforceserver)`[]`<br>
-`credentials` | [PerforceCredentials](#perforcecredentials)`[]`<br>
-`autoSdk` | [AutoSdkWorkspace](#autosdkworkspace)`[]`<br>
-
-## PerforceServer
-
-Name | Description
----- | -----------
-`serverAndPort` | `string`<br>
-`healthCheck` | `boolean`<br>
-`resolveDns` | `boolean`<br>
-`maxConformCount` | `integer`<br>
-`condition` | `string`<br>
-`properties` | `string[]`<br>
-
-## PerforceCredentials
-
-Name | Description
----- | -----------
-`userName` | `string`<br>
-`password` | `string`<br>
-`ticket` | `string`<br>
-
-## AutoSdkWorkspace
-
-Name | Description
----- | -----------
-`name` | `string`<br>
-`properties` | `string[]`<br>
-`userName` | `string`<br>
-`stream` | `string`<br>
-
-## DeviceConfig
-
-Name | Description
----- | -----------
-`platforms` | [DevicePlatformConfig](#deviceplatformconfig)`[]`<br>
-`pools` | [DevicePoolConfig](#devicepoolconfig)`[]`<br>
-
-## DevicePlatformConfig
+## ToolConfig
 
 Name | Description
 ---- | -----------
 `id` | `string`<br>
 `name` | `string`<br>
-`models` | `string[]`<br>
-`legacyNames` | `string[]`<br>
-`legacyPerfSpecHighModel` | `string`<br>
+`description` | `string`<br>
+`category` | `string`<br>
+`group` | `string`<br>
+`platforms` | `string[]`<br>
+`public` | `boolean`<br>
+`showInUgs` | `boolean`<br>
+`showInDashboard` | `boolean`<br>
+`showInToolbox` | `boolean`<br>
+`metadata` | `string` `->` `string`<br>
+`namespaceId` | `string`<br>
+`acl` | [AclConfig](#aclconfig)<br>
 
-## DevicePoolConfig
+## SymbolsConfig
+
+Name | Description
+---- | -----------
+`stores` | [SymbolStoreConfig](#symbolstoreconfig)`[]`<br>
+
+## SymbolStoreConfig
 
 Name | Description
 ---- | -----------
 `id` | `string`<br>
-`name` | `string`<br>
-`poolType` | [DevicePoolType](#devicepooltype-enum)<br>
-`projectIds` | `string[]`<br>
+`namespaceId` | `string`<br>
+`public` | `boolean`<br>
+`acl` | [AclConfig](#aclconfig)<br>
 
-## DevicePoolType (Enum)
-
-The type of device pool
-
-Name | Description
----- | -----------
-`Automation` | Available to CIS jobs
-`Shared` | Shared by users with remote checking and checkouts
-
-## ArtifactTypeConfig
-
-Name | Description
----- | -----------
-`name` | `string`<br>
-`type` | `string`<br>
-`acl` | [AclConfig](#aclconfig)<br>
-`keepCount` | `integer`<br>
-`keepDays` | `integer`<br>
-`namespaceId` | `string`<br>
-
 ## StorageConfig
 
@@ -454,114 +565,3 @@
 `gcDelayHrs` | `number`<br>
 `enableAliases` | `boolean`<br>
 `acl` | [AclConfig](#aclconfig)<br>
-
-## ToolsConfig
-
-Name | Description
----- | -----------
-`tools` | [ToolConfig](#toolconfig)`[]`<br>
-
-## ToolConfig
-
-Name | Description
----- | -----------
-`id` | `string`<br>
-`name` | `string`<br>
-`description` | `string`<br>
-`category` | `string`<br>
-`group` | `string`<br>
-`platforms` | `string[]`<br>
-`public` | `boolean`<br>
-`showInUgs` | `boolean`<br>
-`showInDashboard` | `boolean`<br>
-`showInToolbox` | `boolean`<br>
-`metadata` | `string` `->` `string`<br>
-`namespaceId` | `string`<br>
-`acl` | [AclConfig](#aclconfig)<br>
-
-## ExperimentalConfig
-
-Name | Description
----- | -----------
-`notifications` | [NotificationConfig](Notification.md)`[]`<br>
-
-## SecretsConfig
-
-Name | Description
----- | -----------
-`secrets` | [SecretConfig](#secretconfig)`[]`<br>
-`providerConfigs` | [SecretProviderConfig](#secretproviderconfig)`[]`<br>
-
-## SecretConfig
-
-Name | Description
----- | -----------
-`id` | `string`<br>
-`data` | `string` `->` `string`<br>
-`sources` | [ExternalSecretConfig](#externalsecretconfig)`[]`<br>
-`acl` | [AclConfig](#aclconfig)<br>
-
-## ExternalSecretConfig
-
-Name | Description
----- | -----------
-`providerConfig` | `string`<br>
-`provider` | `string`<br>
-`format` | [ExternalSecretFormat](#externalsecretformat-enum)<br>
-`key` | `string`<br>
-`path` | `string`<br>
-
-## ExternalSecretFormat (Enum)
-
-Name | Description
----- | -----------
-`Text` | 
-`Json` | 
-
-## SecretProviderConfig
-
-Name | Description
----- | -----------
-`name` | `string`<br>
-`provider` | `string`<br>
-`hcpVault` | [HcpVaultConfig](#hcpvaultconfig)<br>
-`azureKeyVault` | [AzureKeyVaultConfig](#azurekeyvaultconfig)<br>
-
-## HcpVaultConfig
-
-Name | Description
----- | -----------
-`preSharedKey` | `string`<br>
-`endPoint` | `string`<br>
-`awsArnRole` | `string`<br>
-`awsIamServerId` | `string`<br>
-`role` | `string`<br>
-`credentials` | [HcpVaultCredentialsType](#hcpvaultcredentialstype-enum)<br>
-
-## HcpVaultCredentialsType (Enum)
-
-Name | Description
----- | -----------
-`PreSharedKey` | 
-`AwsAuth` | 
-
-## AzureKeyVaultConfig
-
-Name | Description
----- | -----------
-`vaultUri` | `string`<br>
-
-## SymbolsConfig
-
-Name | Description
----- | -----------
-`stores` | [SymbolStoreConfig](#symbolstoreconfig)`[]`<br>
-
-## SymbolStoreConfig
-
-Name | Description
----- | -----------
-`id` | `string`<br>
-`namespaceId` | `string`<br>
-`public` | `boolean`<br>
-`acl` | [AclConfig](#aclconfig)<br>
diff --git a/Engine/Source/Programs/Horde/Docs/Deployment/ServerSettings.md b/Engine/Source/Programs/Horde/Docs/Deployment/ServerSettings.md
--- a/Engine/Source/Programs/Horde/Docs/Deployment/ServerSettings.md
+++ b/Engine/Source/Programs/Horde/Docs/Deployment/ServerSettings.md
@@ -127,13 +127,13 @@
 ---- | -----------
 `analytics` | [AnalyticsServerConfig](#analyticsserverconfig)<br>Configuration for the analytics plugin
 `exampleplugin` | [ExamplePluginServerConfig](#examplepluginserverconfig)<br>Configuration for the exampleplugin plugin
-`compute` | [ComputeServerConfig](#computeserverconfig)<br>Configuration for the compute plugin
 `build` | [BuildServerConfig](#buildserverconfig)<br>Configuration for the build plugin
-`storage` | [StorageServerConfig](#storageserverconfig)<br>Configuration for the storage plugin
-`tools` | [ToolsServerConfig](#toolsserverconfig)<br>Configuration for the tools plugin
-`experimental` | [ExperimentalServerConfig](#experimentalserverconfig)<br>Configuration for the experimental plugin
 `secrets` | [SecretsServerConfig](#secretsserverconfig)<br>Configuration for the secrets plugin
+`experimental` | [ExperimentalServerConfig](#experimentalserverconfig)<br>Configuration for the experimental plugin
+`compute` | [ComputeServerConfig](#computeserverconfig)<br>Configuration for the compute plugin
+`tools` | [ToolsServerConfig](#toolsserverconfig)<br>Configuration for the tools plugin
 `symbols` | [PluginServerConfig](#pluginserverconfig)<br>Configuration for the symbols plugin
+`storage` | [StorageServerConfig](#storageserverconfig)<br>Configuration for the storage plugin
 
 ## AnalyticsServerConfig
 
@@ -170,48 +170,6 @@
 ---- | -----------
 `enabled` | `boolean`<br>
 
-## ComputeServerConfig
-
-Name | Description
----- | -----------
-`enableUpgradeTasks` | `boolean`<br>
-`withAws` | `boolean`<br>
-`awsRegions` | `string[]`<br>
-`awsAutoScalingQueueUrls` | `string[]`<br>
-`fleetManagerV2` | [FleetManagerType](#fleetmanagertype-enum)<br>
-`fleetManagerV2Config` | `object`<br>
-`autoEnrollAgents` | `boolean`<br>
-`defaultAgentPoolSizeStrategy` | [PoolSizeStrategy](#poolsizestrategy-enum)<br>
-`agentPoolScaleOutCooldownSeconds` | `integer`<br>
-`agentPoolScaleInCooldownSeconds` | `integer`<br>
-`sessionConflictEnabled` | `boolean`<br>
-`sessionConflictThreshold` | `integer`<br>
-`sessionConflictWindowSeconds` | `integer`<br>
-`computeTunnelPort` | `integer`<br>
-`computeTunnelAddress` | `string`<br>
-`enabled` | `boolean`<br>
-
-## FleetManagerType (Enum)
-
-Name | Description
----- | -----------
-`Default` | 
-`NoOp` | 
-`Aws` | 
-`AwsReuse` | 
-`AwsRecycle` | 
-`AwsAsg` | 
-
-## PoolSizeStrategy (Enum)
-
-Name | Description
----- | -----------
-`LeaseUtilization` | 
-`JobQueue` | 
-`NoOp` | 
-`ComputeQueueAwsMetric` | 
-`LeaseUtilizationAwsMetric` | 
-
 ## BuildServerConfig
 
 Name | Description
@@ -348,56 +306,60 @@
 `maxChildCount` | `integer`<br>Maximum number of children in each node
 `sliceThreshold` | `integer`<br>Threshold hash value for splitting interior nodes
 
-## StorageServerConfig
+## SecretsServerConfig
 
 Name | Description
 ---- | -----------
-`bundleCacheDir` | `string`<br>
-`bundleCacheSize` | `string`<br>
-`bundleCacheSizeBytes` | `integer`<br>
-`backends` | [BackendConfig](#backendconfig)`[]`<br>
+`withAws` | `boolean`<br>
 `enabled` | `boolean`<br>
 
-## BackendConfig
+## ExperimentalServerConfig
 
 Name | Description
 ---- | -----------
-`id` | `string`<br>
-`base` | `string`<br>
-`secondary` | `string`<br>
-`type` | [StorageBackendType](#storagebackendtype-enum)<br>
-`baseDir` | `string`<br>
-`awsBucketName` | `string`<br>
-`awsBucketPath` | `string`<br>
-`awsCredentials` | [AwsCredentialsType](#awscredentialstype-enum)<br>
-`awsRole` | `string`<br>
-`awsProfile` | `string`<br>
-`awsRegion` | `string`<br>
-`azureConnectionString` | `string`<br>
-`azureContainerName` | `string`<br>
-`relayServer` | `string`<br>
-`relayToken` | `string`<br>
-`gcsBucketName` | `string`<br>
-`gcsBucketPath` | `string`<br>
+`enabled` | `boolean`<br>
 
-## StorageBackendType (Enum)
+## ComputeServerConfig
 
 Name | Description
 ---- | -----------
-`FileSystem` | 
-`Aws` | 
-`Azure` | 
-`Gcs` | 
-`Memory` | 
+`enableUpgradeTasks` | `boolean`<br>
+`withAws` | `boolean`<br>
+`awsRegions` | `string[]`<br>
+`awsAutoScalingQueueUrls` | `string[]`<br>
+`fleetManagerV2` | [FleetManagerType](#fleetmanagertype-enum)<br>
+`fleetManagerV2Config` | `object`<br>
+`autoEnrollAgents` | `boolean`<br>
+`defaultAgentPoolSizeStrategy` | [PoolSizeStrategy](#poolsizestrategy-enum)<br>
+`agentPoolScaleOutCooldownSeconds` | `integer`<br>
+`agentPoolScaleInCooldownSeconds` | `integer`<br>
+`sessionConflictEnabled` | `boolean`<br>
+`sessionConflictThreshold` | `integer`<br>
+`sessionConflictWindowSeconds` | `integer`<br>
+`computeTunnelPort` | `integer`<br>
+`computeTunnelAddress` | `string`<br>
+`enabled` | `boolean`<br>
 
-## AwsCredentialsType (Enum)
+## FleetManagerType (Enum)
 
 Name | Description
 ---- | -----------
 `Default` | 
-`Profile` | 
-`AssumeRole` | 
-`AssumeRoleWebIdentity` | 
+`NoOp` | 
+`Aws` | 
+`AwsReuse` | 
+`AwsRecycle` | 
+`AwsAsg` | 
 
+## PoolSizeStrategy (Enum)
+
+Name | Description
+---- | -----------
+`LeaseUtilization` | 
+`JobQueue` | 
+`NoOp` | 
+`ComputeQueueAwsMetric` | 
+`LeaseUtilizationAwsMetric` | 
+
 ## ToolsServerConfig
 
@@ -460,21 +422,59 @@
 `excludeActions` | `string[]`<br>
 `extends` | `string[]`<br>
 
-## ExperimentalServerConfig
-
-Name | Description
----- | -----------
-`enabled` | `boolean`<br>
-
-## SecretsServerConfig
-
-Name | Description
----- | -----------
-`withAws` | `boolean`<br>
-`enabled` | `boolean`<br>
-
 ## PluginServerConfig
 
 Name | Description
 ---- | -----------
 `enabled` | `boolean`<br>
+
+## StorageServerConfig
+
+Name | Description
+---- | -----------
+`bundleCacheDir` | `string`<br>
+`bundleCacheSize` | `string`<br>
+`bundleCacheSizeBytes` | `integer`<br>
+`backends` | [BackendConfig](#backendconfig)`[]`<br>
+`enabled` | `boolean`<br>
+
+## BackendConfig
+
+Name | Description
+---- | -----------
+`id` | `string`<br>
+`base` | `string`<br>
+`secondary` | `string`<br>
+`type` | [StorageBackendType](#storagebackendtype-enum)<br>
+`baseDir` | `string`<br>
+`awsBucketName` | `string`<br>
+`awsBucketPath` | `string`<br>
+`awsCredentials` | [AwsCredentialsType](#awscredentialstype-enum)<br>
+`awsRole` | `string`<br>
+`awsProfile` | `string`<br>
+`awsRegion` | `string`<br>
+`azureConnectionString` | `string`<br>
+`azureContainerName` | `string`<br>
+`relayServer` | `string`<br>
+`relayToken` | `string`<br>
+`gcsBucketName` | `string`<br>
+`gcsBucketPath` | `string`<br>
+
+## StorageBackendType (Enum)
+
+Name | Description
+---- | -----------
+`FileSystem` | 
+`Aws` | 
+`Azure` | 
+`Gcs` | 
+`Memory` | 
+
+## AwsCredentialsType (Enum)
+
+Name | Description
+---- | -----------
+`Default` | 
+`Profile` | 
+`AssumeRole` | 
+`AssumeRoleWebIdentity` | 
