Horde: Suppress Slack [at]-mentions for low-context issue types (Hashed, Gauntlet)
Persist MaxSuspectRank on spans; when all suspects tied at rank 0, list
without [at]-mentions and post explanatory note. Configurable via
LowContextFingerprintTypes in BuildConfig (empty array disables).

diff --git a/Engine/Source/Programs/Horde/Plugins/Build/HordeServer.Build/BuildConfig.cs b/Engine/Source/Programs/Horde/Plugins/Build/HordeServer.Build/BuildConfig.cs
--- a/Engine/Source/Programs/Horde/Plugins/Build/HordeServer.Build/BuildConfig.cs
+++ b/Engine/Source/Programs/Horde/Plugins/Build/HordeServer.Build/BuildConfig.cs
@@ -195,6 +195,14 @@ namespace HordeServer
 		/// </summary>
 		public bool EnableExpandedIgnorePatterns { get; set; } = true;
 
+		/// <summary>
+		/// Fingerprint type prefixes considered "low context" for suspect mention suppression.
+		/// When an issue's fingerprint type starts with any of these prefixes AND all suspects are tied
+		/// at rank 0, @ mentions are suppressed in Slack. Suspects are still listed without mentions.
+		/// Set to an empty array to disable suppression entirely.
+		/// </summary>
+		public List<string> LowContextFingerprintTypes { get; set; } = new List<string> { "Hashed", "Gauntlet" };
+
 		private readonly Dictionary<ProjectId, ProjectConfig> _projectLookup = new Dictionary<ProjectId, ProjectConfig>();
 		private readonly Dictionary<StreamId, StreamConfig> _streamLookup = new Dictionary<StreamId, StreamConfig>();
 
diff --git a/Engine/Source/Programs/Horde/Plugins/Build/HordeServer.Build/Issues/IIssueCollection.cs b/Engine/Source/Programs/Horde/Plugins/Build/HordeServer.Build/Issues/IIssueCollection.cs
--- a/Engine/Source/Programs/Horde/Plugins/Build/HordeServer.Build/Issues/IIssueCollection.cs
+++ b/Engine/Source/Programs/Horde/Plugins/Build/HordeServer.Build/Issues/IIssueCollection.cs
@@ -135,10 +135,11 @@ namespace HordeServer.Issues
 		/// <param name="newFailure">New failed step</param>
 		/// <param name="newNextSuccess">New next successful step</param>
 		/// <param name="newSuspects">New suspects for the span</param>
+		/// <param name="newMaxSuspectRank">New maximum suspect rank for the span</param>
 		/// <param name="newIssueId">The new issue id for this span</param>
 		/// <param name="cancellationToken">Cancellation token for the operation</param>
 		/// <returns>The updated span, or null on failure</returns>
-		Task<IIssueSpan?> TryUpdateSpanAsync(IIssueSpan span, NewIssueStepData? newLastSuccess = null, NewIssueStepData? newFailure = null, NewIssueStepData? newNextSuccess = null, List<NewIssueSpanSuspectData>? newSuspects = null, int? newIssueId = null, CancellationToken cancellationToken = default);
+		Task<IIssueSpan?> TryUpdateSpanAsync(IIssueSpan span, NewIssueStepData? newLastSuccess = null, NewIssueStepData? newFailure = null, NewIssueStepData? newNextSuccess = null, List<NewIssueSpanSuspectData>? newSuspects = null, int? newMaxSuspectRank = null, int? newIssueId = null, CancellationToken cancellationToken = default);
 
 		/// <summary>
 		/// Gets all the spans for a particular issue
@@ -575,6 +576,9 @@ namespace HordeServer.Issues
 		/// <inheritdoc cref="IIssueSpan.Suspects"/>
 		public List<NewIssueSpanSuspectData> Suspects { get; set; } = new List<NewIssueSpanSuspectData>();
 
+		/// <inheritdoc cref="IIssueSpan.MaxSuspectRank"/>
+		public int? MaxSuspectRank { get; set; }
+
 		/// <summary>
 		/// Constructor
 		/// </summary>
diff --git a/Engine/Source/Programs/Horde/Plugins/Build/HordeServer.Build/Issues/IIssueSpan.cs b/Engine/Source/Programs/Horde/Plugins/Build/HordeServer.Build/Issues/IIssueSpan.cs
--- a/Engine/Source/Programs/Horde/Plugins/Build/HordeServer.Build/Issues/IIssueSpan.cs
+++ b/Engine/Source/Programs/Horde/Plugins/Build/HordeServer.Build/Issues/IIssueSpan.cs
@@ -74,6 +74,13 @@ namespace HordeServer.Issues
 		/// </summary>
 		public IReadOnlyList<IIssueSpanSuspect> Suspects { get; }
 
+		/// <summary>
+		/// The maximum rank achieved during suspect scoring for this span. A value of 0 means all suspects
+		/// are tied with no meaningful differentiation (no file/symbol/note matches). Null if not yet computed
+		/// or for spans created before this property was added.
+		/// </summary>
+		public int? MaxSuspectRank { get; }
+
 		/// <summary>
 		/// The issue that this span is assigned to
 		/// </summary>
diff --git a/Engine/Source/Programs/Horde/Plugins/Build/HordeServer.Build/Issues/IssueCollection.cs b/Engine/Source/Programs/Horde/Plugins/Build/HordeServer.Build/Issues/IssueCollection.cs
--- a/Engine/Source/Programs/Horde/Plugins/Build/HordeServer.Build/Issues/IssueCollection.cs
+++ b/Engine/Source/Programs/Horde/Plugins/Build/HordeServer.Build/Issues/IssueCollection.cs
@@ -492,6 +492,10 @@ namespace HordeServer.Issues
 			bool IIssueSpan.PromoteByDefault => PromoteByDefault ?? NotifySuspectsDeprecated;
 
 			public List<IssueSpanSuspectDocument> Suspects { get; set; }
+
+			[BsonIgnoreIfNull]
+			public int? MaxSuspectRank { get; set; }
+
 			public int IssueId { get; set; }
 			public int UpdateIndex { get; set; }
 
@@ -500,6 +504,7 @@ namespace HordeServer.Issues
 			IIssueStep IIssueSpan.LastFailure => LastFailure;
 			IIssueStep? IIssueSpan.NextSuccess => NextSuccess;
 			IReadOnlyList<IIssueSpanSuspect> IIssueSpan.Suspects => Suspects;
+			int? IIssueSpan.MaxSuspectRank => MaxSuspectRank;
 			IIssueFingerprint IIssueSpan.Fingerprint => Fingerprint;
 
 			private IssueSpanDocument()
@@ -534,6 +539,7 @@ namespace HordeServer.Issues
 				}
 				PromoteByDefault = newSpan.FirstFailure.PromoteByDefault;
 				Suspects = newSpan.Suspects.ConvertAll(x => new IssueSpanSuspectDocument(x));
+				MaxSuspectRank = newSpan.MaxSuspectRank;
 				IssueId = issueId;
 			}
 		}
@@ -1761,7 +1767,7 @@ namespace HordeServer.Issues
 		}
 
 		/// <inheritdoc/>
-		public async Task<IIssueSpan?> TryUpdateSpanAsync(IIssueSpan span, NewIssueStepData? newLastSuccess = null, NewIssueStepData? newFailure = null, NewIssueStepData? newNextSuccess = null, List<NewIssueSpanSuspectData>? newSuspects = null, int? newIssueId = null, CancellationToken cancellationToken = default)
+		public async Task<IIssueSpan?> TryUpdateSpanAsync(IIssueSpan span, NewIssueStepData? newLastSuccess = null, NewIssueStepData? newFailure = null, NewIssueStepData? newNextSuccess = null, List<NewIssueSpanSuspectData>? newSuspects = null, int? newMaxSuspectRank = null, int? newIssueId = null, CancellationToken cancellationToken = default)
 		{
 			List<UpdateDefinition<IssueSpanDocument>> updates = new List<UpdateDefinition<IssueSpanDocument>>();
 			if (newLastSuccess != null)
@@ -1793,6 +1799,10 @@ namespace HordeServer.Issues
 			{
 				updates.Add(Builders<IssueSpanDocument>.Update.Set(x => x.Suspects, newSuspects.ConvertAll(x => new IssueSpanSuspectDocument(x))));
 			}
+			if (newMaxSuspectRank != null)
+			{
+				updates.Add(Builders<IssueSpanDocument>.Update.Set(x => x.MaxSuspectRank, newMaxSuspectRank.Value));
+			}
 			if (newIssueId != null)
 			{
 				updates.Add(Builders<IssueSpanDocument>.Update.Set(x => x.IssueId, newIssueId.Value));
diff --git a/Engine/Source/Programs/Horde/Plugins/Build/HordeServer.Build/Issues/IssueService.cs b/Engine/Source/Programs/Horde/Plugins/Build/HordeServer.Build/Issues/IssueService.cs
--- a/Engine/Source/Programs/Horde/Plugins/Build/HordeServer.Build/Issues/IssueService.cs
+++ b/Engine/Source/Programs/Horde/Plugins/Build/HordeServer.Build/Issues/IssueService.cs
@@ -1043,7 +1043,7 @@ namespace HordeServer.Issues
 				if (prevJob != null)
 				{
 					spanData.LastSuccess = new NewIssueStepData(prevJob);
-					spanData.Suspects = await FindSuspectsForSpanAsync(streamConfig, spanData.Fingerprint, spanData.LastSuccess.CommitId, spanData.LastSuccess.StepTime, spanData.FirstFailure.CommitId, spanData.FirstFailure.StepTime, groupAuditLogger, cancellationToken);
+					(spanData.Suspects, spanData.MaxSuspectRank) = await FindSuspectsForSpanAsync(streamConfig, spanData.Fingerprint, spanData.LastSuccess.CommitId, spanData.LastSuccess.StepTime, spanData.FirstFailure.CommitId, spanData.FirstFailure.StepTime, groupAuditLogger, cancellationToken);
 				}
 
 				IJobStepRef? nextJob = await _jobStepRefs.GetNextStepForNodeAsync(job.StreamId, job.TemplateId, node.Name, job.CommitId, JobStepOutcome.Success, true);
@@ -1384,11 +1384,12 @@ namespace HordeServer.Issues
 		/// <param name="issueLogger">The audit logger for the issue that the span is associated with</param>
 		/// <param name="cancellationToken">Cancellation token for the operation</param>
 		/// <returns>List of suspects</returns>
-		async Task<List<NewIssueSpanSuspectData>> FindSuspectsForSpanAsync(StreamConfig streamConfig, IIssueFingerprint fingerprint, CommitId lastSuccess, DateTime lastSuccessTime, CommitId firstFailure, DateTime firstFailureTime, ILogger? issueLogger, CancellationToken cancellationToken)
+		async Task<(List<NewIssueSpanSuspectData> Suspects, int? MaxRank)> FindSuspectsForSpanAsync(StreamConfig streamConfig, IIssueFingerprint fingerprint, CommitId lastSuccess, DateTime lastSuccessTime, CommitId firstFailure, DateTime firstFailureTime, ILogger? issueLogger, CancellationToken cancellationToken)
 		{
 			using TelemetrySpan span = _tracer.StartActiveSpan($"{nameof(IssueService)}.{nameof(FindSuspectsForSpanAsync)}");
-			
+
 			List<NewIssueSpanSuspectData> suspects = new List<NewIssueSpanSuspectData>();
+			int maxRank = -1;
 
 			try
 			{
@@ -1405,7 +1406,7 @@ namespace HordeServer.Issues
 						if (Math.Abs(delta.TotalDays) >= _buildConfig.CurrentValue.IssueFindSuspectsMaxSuccessDays)
 						{
 							LoggerUtils.MultiLog([(_logger, LogLevel.Warning), (issueLogger, LogLevel.Debug)], "Ignoring commits for {StreamName}: {LastSuccess} and {FirstFailure}, delta is {DeltaDays}", streamConfig.Name, lastSuccess.Name, firstFailure.Name, delta.TotalDays);
-							return suspects;
+							return (suspects, null);
 						}
 					}
 
@@ -1513,8 +1514,6 @@ namespace HordeServer.Issues
 					DateTime scoringStart = DateTime.Now;
 
 					// Get the filter for changes to include in the suspect list
-					int maxRank = -1;
-
 					if (_buildConfig.CurrentValue.IssueFindSuspectsEnableParallelScoring)
 					{
 						// Phase 1: Score all commits in parallel
@@ -1687,7 +1686,7 @@ namespace HordeServer.Issues
 			{
 				LoggerUtils.MultiLog([(_logger, LogLevel.Error), (issueLogger, LogLevel.Information)], ex, "Exception finding suspects for {StreamName}: {LastSuccess} and {FirstFailure}", streamConfig.Name, lastSuccess, firstFailure.Name);
 			}
-			return suspects;
+			return (suspects, maxRank >= 0 ? maxRank : null);
 		}
 
 		static async ValueTask<int> GetChangeScoreAsync(ICommit commit, FileFilter filter, HashSet<string> suspectFiles, HashSet<string> suspectNotes, HashSet<string> suspectSymbols, Tracer tracer, ILogger? issueLogger)
@@ -2070,9 +2069,9 @@ namespace HordeServer.Issues
 				{
 					NewIssueStepData newLastSuccess = new NewIssueStepData(job.CommitId, IssueSeverity.Unspecified, job.Name, job.Id, batch.Id, step.Id, step.StartTimeUtc ?? default, step.LogId, null, false);
 					ILogger issueLogger = _issueCollection.GetLogger(span.IssueId);
-					List<NewIssueSpanSuspectData> newSuspects = await FindSuspectsForSpanAsync(streamConfig, span.Fingerprint, newLastSuccess.CommitId, newLastSuccess.StepTime, span.FirstFailure.CommitId, span.FirstFailure.StepTime, issueLogger, cancellationToken);
+					(List<NewIssueSpanSuspectData> newSuspects, int? newMaxSuspectRank) = await FindSuspectsForSpanAsync(streamConfig, span.Fingerprint, newLastSuccess.CommitId, newLastSuccess.StepTime, span.FirstFailure.CommitId, span.FirstFailure.StepTime, issueLogger, cancellationToken);
 
-					if (await _issueCollection.TryUpdateSpanAsync(span, newLastSuccess: newLastSuccess, newSuspects: newSuspects, cancellationToken: cancellationToken) == null)
+					if (await _issueCollection.TryUpdateSpanAsync(span, newLastSuccess: newLastSuccess, newSuspects: newSuspects, newMaxSuspectRank: newMaxSuspectRank, cancellationToken: cancellationToken) == null)
 					{
 						return false;
 					}
diff --git a/Engine/Source/Programs/Horde/Plugins/Build/HordeServer.Build/Notifications/Sinks/SlackNotificationSink.cs b/Engine/Source/Programs/Horde/Plugins/Build/HordeServer.Build/Notifications/Sinks/SlackNotificationSink.cs
--- a/Engine/Source/Programs/Horde/Plugins/Build/HordeServer.Build/Notifications/Sinks/SlackNotificationSink.cs
+++ b/Engine/Source/Programs/Horde/Plugins/Build/HordeServer.Build/Notifications/Sinks/SlackNotificationSink.cs
@@ -1249,8 +1249,52 @@ namespace HordeServer.Notifications.Sinks
 					else
 					{
 						IGrouping<UserId, IIssueSuspect>[] suspectGroups = suspects.GroupBy(x => x.AuthorId).ToArray();
-						if (suspectGroups.Length > 0 && suspectGroups.Length <= workflow.MaxMentions)
+
+						List<string> lowContextTypes = _buildConfig.CurrentValue.LowContextFingerprintTypes;
+						bool isLowContextType = lowContextTypes.Any(prefix => span.Fingerprint.Type.StartsWith(prefix, StringComparison.Ordinal));
+
+						if (suspectGroups.Length > 0
+							&& suspectGroups.Length <= workflow.MaxMentions
+							&& isLowContextType
+							&& span.MaxSuspectRank != null
+							&& span.MaxSuspectRank.Value == 0)
+						{
+							// PATH 0: Low-context type with all suspects tied at rank 0.
+							// List suspects without @ mentions; post explanatory note.
+							_issueService.Collection.GetLogger(issue.Id).LogInformation(
+								"Suppressing @ mentions for issue with MaxSuspectRank={MaxSuspectRank} ({SuspectCount} suspects, fingerprint type: {FingerprintType})",
+								span.MaxSuspectRank, suspectGroups.Length, span.Fingerprint.Type);
+
+							List<string> suspectList = new List<string>();
+							foreach (IGrouping<UserId, IIssueSuspect> suspectGroup in suspectGroups)
+							{
+								string name = await FormatNameAsync(suspectGroup.Key, cancellationToken);
+								string changes = String.Join(", ", suspectGroup.Select(x => FormatChange(x.CommitId)));
+								suspectList.Add($"{name} ({changes})");
+							}
+
+							string suspectMessage = $"Possibly {StringUtils.FormatList(suspectList, "or")}.";
+
+							Uri? swarmLink = GetSwarmLinkForSpan(span);
+							if (swarmLink != null)
+							{
+								suspectMessage += $" (<{swarmLink}|View changes>)";
+							}
+
+							Uri? robomergeTrackChangesLink = GetRobomergeTrackChangesUrlForSpan(span);
+							if (robomergeTrackChangesLink != null)
+							{
+								suspectMessage += $" (<{robomergeTrackChangesLink}|Track changes>)";
+							}
+
+							await _slackClient.PostMessageToThreadAsync(threadId, suspectMessage, cancellationToken: cancellationToken);
+							await _slackClient.PostMessageToThreadAsync(threadId,
+								$"_This issue does not have enough context (files, symbols) to identify a specific suspect. Mentions suppressed._",
+								cancellationToken: cancellationToken);
+						}
+						else if (suspectGroups.Length > 0 && suspectGroups.Length <= workflow.MaxMentions)
 						{
+							// PATH 1: @ mentions via FormatMentionAsync
 							List<string> suspectList = new List<string>();
 							foreach (IGrouping<UserId, IIssueSuspect> suspectGroup in suspectGroups)
 							{
@@ -1278,6 +1322,7 @@ namespace HordeServer.Notifications.Sinks
 						}
 						else
 						{
+							// PATH 2: Swarm/Commits links only, notifyTriageAlias = true
 							string suspectMessage = String.Empty;
 
 							Uri? swarmLink = GetSwarmLinkForSpan(span);
