Horde: Add debug endpoint for creating test issues with suspects
Allows testing the notification pipeline without requiring real
Perforce commits and job failures.

diff --git a/Engine/Source/Programs/Horde/Plugins/Build/HordeServer.Build/Issues/DebugIssuesController.cs b/Engine/Source/Programs/Horde/Plugins/Build/HordeServer.Build/Issues/DebugIssuesController.cs
new file mode 100644
--- /dev/null
+++ b/Engine/Source/Programs/Horde/Plugins/Build/HordeServer.Build/Issues/DebugIssuesController.cs
@@ -0,0 +1,155 @@
+// Copyright Epic Games, Inc. All Rights Reserved.
+
+using EpicGames.Horde.Commits;
+using EpicGames.Horde.Issues;
+using EpicGames.Horde.Jobs;
+using EpicGames.Horde.Jobs.Templates;
+using EpicGames.Horde.Streams;
+using HordeServer.Jobs;
+using HordeServer.Server;
+using HordeServer.Streams;
+using HordeServer.Users;
+using HordeServer.Utilities;
+using Microsoft.AspNetCore.Authorization;
+using Microsoft.AspNetCore.Http;
+using Microsoft.AspNetCore.Mvc;
+using Microsoft.Extensions.Options;
+
+namespace HordeServer.Issues
+{
+	/// <summary>
+	/// Debug functionality for issues
+	/// </summary>
+	[ApiController]
+	[Authorize]
+	[DebugEndpoint]
+	[Tags("Debug")]
+	public class DebugIssuesController : HordeControllerBase
+	{
+		private readonly IIssueCollection _issueCollection;
+		private readonly IssueService _issueService;
+		private readonly IUserCollection _userCollection;
+		private readonly IOptionsSnapshot<BuildConfig> _buildConfig;
+
+		/// <summary>
+		/// Constructor
+		/// </summary>
+		public DebugIssuesController(
+			IIssueCollection issueCollection,
+			IssueService issueService,
+			IUserCollection userCollection,
+			IOptionsSnapshot<BuildConfig> buildConfig)
+		{
+			_issueCollection = issueCollection;
+			_issueService = issueService;
+			_userCollection = userCollection;
+			_buildConfig = buildConfig;
+		}
+
+		/// <summary>
+		/// Creates a test issue with synthetic suspects and triggers the notification pipeline.
+		/// Useful for testing Slack DM notifications without requiring real Perforce commits and job failures.
+		/// </summary>
+		[HttpPost]
+		[Route("/api/v1/debug/issues/create-test-issue")]
+		public async Task<ActionResult> CreateTestIssueAsync(
+			[FromQuery] StreamId streamId,
+			[FromQuery] int commitId,
+			[FromQuery] string[] suspectLogins,
+			[FromQuery] int[] suspectCommitIds,
+			[FromQuery] TemplateId? templateId = null,
+			[FromQuery] string? nodeName = null,
+			[FromQuery] string? summary = null,
+			[FromQuery] bool promoted = true)
+		{
+			if (!_buildConfig.Value.Authorize(ServerAclAction.Debug, User))
+			{
+				return Forbid(ServerAclAction.Debug);
+			}
+
+			if (suspectLogins.Length != suspectCommitIds.Length)
+			{
+				return BadRequest("suspectLogins and suspectCommitIds must have the same length");
+			}
+
+			if (!_buildConfig.Value.TryGetStream(streamId, out StreamConfig? streamConfig))
+			{
+				return NotFound($"Stream '{streamId}' not found");
+			}
+
+			// Resolve users
+			List<IUser> resolvedUsers = new List<IUser>();
+			List<string> notFoundLogins = new List<string>();
+			for (int i = 0; i < suspectLogins.Length; i++)
+			{
+				IUser? user = await _userCollection.FindUserByLoginAsync(suspectLogins[i], HttpContext.RequestAborted);
+				if (user == null)
+				{
+					notFoundLogins.Add(suspectLogins[i]);
+				}
+				else
+				{
+					resolvedUsers.Add(user);
+				}
+			}
+
+			if (notFoundLogins.Count > 0)
+			{
+				return BadRequest($"Could not resolve users: {String.Join(", ", notFoundLogins)}");
+			}
+
+			// Create the issue
+			string issueSummary = summary ?? "Debug test issue";
+			IIssue issue = await _issueCollection.AddIssueAsync(issueSummary, HttpContext.RequestAborted);
+
+			// Promote the issue if requested
+			if (promoted)
+			{
+				issue = await _issueCollection.TryUpdateIssueAsync(issue, null, new UpdateIssueOptions(Promoted: true), HttpContext.RequestAborted) ?? issue;
+			}
+
+			// Build fingerprint
+			NewIssueFingerprint fingerprint = new NewIssueFingerprint("Compile", "{Summary}", "Code");
+
+			// Build step data
+			TemplateId effectiveTemplateId = templateId ?? new TemplateId("test-executor");
+			string effectiveNodeName = nodeName ?? "Test Node";
+			NewIssueStepData stepData = new NewIssueStepData(
+				CommitIdWithOrder.FromPerforceChange(commitId),
+				IssueSeverity.Error,
+				"Debug Test Job",
+				JobIdUtils.GenerateNewId(),
+				JobStepBatchId.GenerateNewId(),
+				JobStepId.GenerateNewId(),
+				DateTime.UtcNow,
+				null,
+				null,
+				promoted
+			);
+
+			// Build suspects
+			List<NewIssueSpanSuspectData> suspects = new List<NewIssueSpanSuspectData>();
+			for (int i = 0; i < resolvedUsers.Count; i++)
+			{
+				suspects.Add(new NewIssueSpanSuspectData(CommitIdWithOrder.FromPerforceChange(suspectCommitIds[i]), resolvedUsers[i].Id));
+			}
+
+			// Build and add span
+			NewIssueSpanData spanData = new NewIssueSpanData(streamId, streamConfig.Name, effectiveTemplateId, effectiveNodeName, fingerprint, stepData);
+			spanData.Suspects = suspects;
+			IIssueSpan span = await _issueCollection.AddSpanAsync(issue.Id, spanData, HttpContext.RequestAborted);
+
+			// Add step to span
+			await _issueCollection.AddStepAsync(span.Id, stepData, HttpContext.RequestAborted);
+
+			// Trigger derived data update + notification pipeline
+			await _issueService.UpdateDerivedDataAsync(issue.Id, HttpContext.RequestAborted);
+
+			return Ok(new
+			{
+				IssueId = issue.Id,
+				Suspects = resolvedUsers.Select((u, i) => new { Login = suspectLogins[i], UserId = u.Id.ToString(), CommitId = suspectCommitIds[i] })
+			});
+		}
+	}
+}
diff --git a/Engine/Source/Programs/Horde/Plugins/Build/HordeServer.Build/Issues/IssueService.cs b/Engine/Source/Programs/Horde/Plugins/Build/HordeServer.Build/Issues/IssueService.cs
--- a/Engine/Source/Programs/Horde/Plugins/Build/HordeServer.Build/Issues/IssueService.cs
+++ b/Engine/Source/Programs/Horde/Plugins/Build/HordeServer.Build/Issues/IssueService.cs
@@ -1074,6 +1074,9 @@ namespace HordeServer.Issues
 			return true;
 		}
 
+		internal Task<IIssue?> UpdateDerivedDataAsync(int issueId, CancellationToken cancellationToken)
+			=> UpdateIssueDerivedDataAsync(issueId, cancellationToken);
+
 		async Task<IIssue?> UpdateIssueDerivedDataAsync(int issueId, CancellationToken cancellationToken)
 		{
 			IIssue? issue = await _issueCollection.GetIssueAsync(issueId, cancellationToken);
