Horde: Add fingerprintType parameter to debug create-test-issue endpoint

diff --git a/Engine/Source/Programs/Horde/Plugins/Build/HordeServer.Build/Issues/DebugIssuesController.cs b/Engine/Source/Programs/Horde/Plugins/Build/HordeServer.Build/Issues/DebugIssuesController.cs
--- a/Engine/Source/Programs/Horde/Plugins/Build/HordeServer.Build/Issues/DebugIssuesController.cs
+++ b/Engine/Source/Programs/Horde/Plugins/Build/HordeServer.Build/Issues/DebugIssuesController.cs
@@ -60,6 +60,7 @@ namespace HordeServer.Issues
 			[FromQuery] TemplateId? templateId = null,
 			[FromQuery] string? nodeName = null,
 			[FromQuery] string? summary = null,
+			[FromQuery] string? fingerprintType = null,
 			[FromQuery] bool promoted = true)
 		{
 			if (!_buildConfig.Value.Authorize(ServerAclAction.Debug, User))
@@ -109,7 +110,7 @@ namespace HordeServer.Issues
 			}
 
 			// Build fingerprint
-			NewIssueFingerprint fingerprint = new NewIssueFingerprint("Compile", "{Summary}", "Code");
+			NewIssueFingerprint fingerprint = new NewIssueFingerprint(fingerprintType ?? "Compile", "{Summary}", "Code");
 
 			// Build step data
 			TemplateId effectiveTemplateId = templateId ?? new TemplateId("test-executor");
