Horde: Add jobId, batchId, stepId, workflowId, fingerprintType params to debug create-test-issue endpoint

diff --git a/Engine/Source/Programs/Horde/Plugins/Build/HordeServer.Build/Issues/DebugIssuesController.cs b/Engine/Source/Programs/Horde/Plugins/Build/HordeServer.Build/Issues/DebugIssuesController.cs
--- a/Engine/Source/Programs/Horde/Plugins/Build/HordeServer.Build/Issues/DebugIssuesController.cs
+++ b/Engine/Source/Programs/Horde/Plugins/Build/HordeServer.Build/Issues/DebugIssuesController.cs
@@ -3,6 +3,7 @@
 using EpicGames.Horde.Commits;
 using EpicGames.Horde.Issues;
 using EpicGames.Horde.Jobs;
+using EpicGames.Horde.Jobs.Graphs;
 using EpicGames.Horde.Jobs.Templates;
 using EpicGames.Horde.Streams;
 using HordeServer.Jobs;
@@ -61,6 +62,10 @@ namespace HordeServer.Issues
 			[FromQuery] string? nodeName = null,
 			[FromQuery] string? summary = null,
 			[FromQuery] string? fingerprintType = null,
+			[FromQuery] string? jobId = null,
+			[FromQuery] string? batchId = null,
+			[FromQuery] string? stepId = null,
+			[FromQuery] string? workflowId = null,
 			[FromQuery] bool promoted = true)
 		{
 			if (!_buildConfig.Value.Authorize(ServerAclAction.Debug, User))
@@ -112,6 +117,19 @@ namespace HordeServer.Issues
 			// Build fingerprint
 			NewIssueFingerprint fingerprint = new NewIssueFingerprint(fingerprintType ?? "Compile", "{Summary}", "Code");
 
+			// Use real job/batch/step IDs when provided, otherwise generate synthetic ones
+			JobId effectiveJobId = jobId != null ? JobId.Parse(jobId) : JobIdUtils.GenerateNewId();
+			JobStepBatchId effectiveBatchId = batchId != null ? JobStepBatchId.Parse(batchId) : JobStepBatchId.GenerateNewId();
+			JobStepId effectiveStepId = stepId != null ? JobStepId.Parse(stepId) : JobStepId.GenerateNewId();
+
+			// Build annotations (for workflow routing)
+			NodeAnnotations? annotations = null;
+			if (workflowId != null)
+			{
+				annotations = new NodeAnnotations();
+				annotations.WorkflowId = new WorkflowId(workflowId);
+			}
+
 			// Build step data
 			TemplateId effectiveTemplateId = templateId ?? new TemplateId("test-executor");
 			string effectiveNodeName = nodeName ?? "Test Node";
@@ -119,12 +137,12 @@ namespace HordeServer.Issues
 				CommitIdWithOrder.FromPerforceChange(commitId),
 				IssueSeverity.Error,
 				"Debug Test Job",
-				JobIdUtils.GenerateNewId(),
-				JobStepBatchId.GenerateNewId(),
-				JobStepId.GenerateNewId(),
+				effectiveJobId,
+				effectiveBatchId,
+				effectiveStepId,
 				DateTime.UtcNow,
 				null,
-				null,
+				annotations,
 				promoted
 			);
 
