Horde: Add lastSuccessCommitId and maxSuspectRank to debug create-test-issue endpoint
#rnx

diff --git a/Engine/Source/Programs/Horde/Plugins/Build/HordeServer.Build/Issues/DebugIssuesController.cs b/Engine/Source/Programs/Horde/Plugins/Build/HordeServer.Build/Issues/DebugIssuesController.cs
--- a/Engine/Source/Programs/Horde/Plugins/Build/HordeServer.Build/Issues/DebugIssuesController.cs
+++ b/Engine/Source/Programs/Horde/Plugins/Build/HordeServer.Build/Issues/DebugIssuesController.cs
@@ -66,6 +66,8 @@ namespace HordeServer.Issues
 			[FromQuery] string? batchId = null,
 			[FromQuery] string? stepId = null,
 			[FromQuery] string? workflowId = null,
+			[FromQuery] int? lastSuccessCommitId = null,
+			[FromQuery] int? maxSuspectRank = null,
 			[FromQuery] bool promoted = true)
 		{
 			if (!_buildConfig.Value.Authorize(ServerAclAction.Debug, User))
@@ -156,6 +158,24 @@ namespace HordeServer.Issues
 			// Build and add span
 			NewIssueSpanData spanData = new NewIssueSpanData(streamId, streamConfig.Name, effectiveTemplateId, effectiveNodeName, fingerprint, stepData);
 			spanData.Suspects = suspects;
+			spanData.MaxSuspectRank = maxSuspectRank;
+
+			// Set LastSuccess if provided (enables CL range display in Slack instead of "No previous successes")
+			if (lastSuccessCommitId != null)
+			{
+				spanData.LastSuccess = new NewIssueStepData(
+					CommitIdWithOrder.FromPerforceChange(lastSuccessCommitId.Value),
+					IssueSeverity.Error,
+					"Debug Test Job",
+					effectiveJobId,
+					effectiveBatchId,
+					effectiveStepId,
+					DateTime.UtcNow.AddMinutes(-10),
+					null,
+					null,
+					false
+				);
+			}
 			IIssueSpan span = await _issueCollection.AddSpanAsync(issue.Id, spanData, HttpContext.RequestAborted);
 
 			// Add step to span
