Horde: Fix crash in GetIssueEventsAsync when logIds contains invalid values

diff --git a/Engine/Source/Programs/Horde/Plugins/Build/HordeServer.Build/Issues/IssuesController.cs b/Engine/Source/Programs/Horde/Plugins/Build/HordeServer.Build/Issues/IssuesController.cs
--- a/Engine/Source/Programs/Horde/Plugins/Build/HordeServer.Build/Issues/IssuesController.cs
+++ b/Engine/Source/Programs/Horde/Plugins/Build/HordeServer.Build/Issues/IssuesController.cs
@@ -759,7 +759,13 @@ namespace HordeServer.Issues
 			}
 			if (logIds != null)
 			{
-				logIdValues.UnionWith(logIds.Select(x => LogId.Parse(x)));
+				foreach (string logIdStr in logIds)
+				{
+					if (LogId.TryParse(logIdStr, out LogId parsedLogId))
+					{
+						logIdValues.Add(parsedLogId);
+					}
+				}
 			}
 
 			IReadOnlyList<IIssueSpan> spans = await _issueCollection.FindSpansAsync(issueId, cancellationToken);
