Horde: Updating schema documentation

#rnx

diff --git a/Engine/Source/Programs/Horde/Docs/Config/Schema/AclActions.md b/Engine/Source/Programs/Horde/Docs/Config/Schema/AclActions.md
--- a/Engine/Source/Programs/Horde/Docs/Config/Schema/AclActions.md
+++ b/Engine/Source/Programs/Horde/Docs/Config/Schema/AclActions.md
@@ -37,6 +37,48 @@
 | `DeleteAccount` |  |
 | `ViewAccount` |  |
 
+## Artifacts
+
+| Name | Description |
+| ---- | ----------- |
+| `ReadArtifact` |  |
+| `WriteArtifact` |  |
+| `DeleteArtifact` |  |
+| `UploadArtifact` |  |
+| `DownloadArtifact` |  |
+
+## Bisect
+
+| Name | Description |
+| ---- | ----------- |
+| `CreateBisectTask` |  |
+| `UpdateBisectTask` |  |
+| `ViewBisectTask` |  |
+
+## BuildHealth
+
+| Name | Description |
+| ---- | ----------- |
+| `AddUpdateDeleteBuildHealthFilter` |  |
+
+## Devices
+
+| Name | Description |
+| ---- | ----------- |
+| `DeviceRead` |  |
+| `DeviceWrite` |  |
+
+## Jobs
+
+| Name | Description |
+| ---- | ----------- |
+| `CreateJob` |  |
+| `UpdateJob` |  |
+| `DeleteJob` |  |
+| `ExecuteJob` |  |
+| `RetryJobStep` |  |
+| `ViewJob` |  |
+
 ## Notices
 
 | Name | Description |
@@ -45,6 +87,28 @@
 | `UpdateNotice` |  |
 | `DeleteNotice` |  |
 
+## Notifications
+
+| Name | Description |
+| ---- | ----------- |
+| `CreateSubscription` |  |
+
+## Projects
+
+| Name | Description |
+| ---- | ----------- |
+| `CreateProject` |  |
+| `DeleteProject` |  |
+| `UpdateProject` |  |
+| `ViewProject` |  |
+
+## Replicators
+
+| Name | Description |
+| ---- | ----------- |
+| `UpdateReplicator` |  |
+| `ViewReplicator` |  |
+
 ## ServiceAccounts
 
 | Name | Description |
@@ -54,6 +118,24 @@
 | `DeleteAccount` |  |
 | `ViewAccount` |  |
 
+## Streams
+
+| Name | Description |
+| ---- | ----------- |
+| `CreateStream` |  |
+| `UpdateStream` |  |
+| `DeleteStream` |  |
+| `ViewStream` |  |
+| `ViewChanges` |  |
+| `ViewTemplate` |  |
+
+## TestData
+
+| Name | Description |
+| ---- | ----------- |
+| `TestHealthRead` |  |
+| `TestHealthWrite` |  |
+
 ## Accounts
 
 | Name | Description |
@@ -63,6 +145,86 @@
 | `DeleteAccount` |  |
 | `ViewAccount` |  |
 
+## Agents
+
+| Name | Description |
+| ---- | ----------- |
+| `CreateAgent` |  |
+| `CreateWorkstationAgent` |  |
+| `UpdateAgent` |  |
+| `DeleteAgent` |  |
+| `ViewAgent` |  |
+| `ListAgents` |  |
+
+## Artifacts
+
+| Name | Description |
+| ---- | ----------- |
+| `ReadArtifact` |  |
+| `WriteArtifact` |  |
+| `DeleteArtifact` |  |
+| `UploadArtifact` |  |
+| `DownloadArtifact` |  |
+
+## Bisect
+
+| Name | Description |
+| ---- | ----------- |
+| `CreateBisectTask` |  |
+| `UpdateBisectTask` |  |
+| `ViewBisectTask` |  |
+
+## BuildHealth
+
+| Name | Description |
+| ---- | ----------- |
+| `AddUpdateDeleteBuildHealthFilter` |  |
+
+## Compute
+
+| Name | Description |
+| ---- | ----------- |
+| `AddComputeTasks` |  |
+| `GetComputeTasks` |  |
+| `UbaCacheRead` |  |
+| `UbaCacheWrite` |  |
+
+## Devices
+
+| Name | Description |
+| ---- | ----------- |
+| `DeviceRead` |  |
+| `DeviceWrite` |  |
+
+## Jobs
+
+| Name | Description |
+| ---- | ----------- |
+| `CreateJob` |  |
+| `UpdateJob` |  |
+| `DeleteJob` |  |
+| `ExecuteJob` |  |
+| `RetryJobStep` |  |
+| `ViewJob` |  |
+
+## Leases
+
+| Name | Description |
+| ---- | ----------- |
+| `ViewLeases` |  |
+| `ViewLeaseTasks` |  |
+
+## Logs
+
+| Name | Description |
+| ---- | ----------- |
+| `CreateLog` |  |
+| `UpdateLog` |  |
+| `ViewLog` |  |
+| `WriteLogData` |  |
+| `CreateEvent` |  |
+| `ViewEvent` |  |
+
 ## Notices
 
 | Name | Description |
@@ -71,6 +233,38 @@
 | `UpdateNotice` |  |
 | `DeleteNotice` |  |
 
+## Notifications
+
+| Name | Description |
+| ---- | ----------- |
+| `CreateSubscription` |  |
+
+## Pools
+
+| Name | Description |
+| ---- | ----------- |
+| `CreatePool` |  |
+| `UpdatePool` |  |
+| `DeletePool` |  |
+| `ViewPool` |  |
+| `ListPools` |  |
+
+## Projects
+
+| Name | Description |
+| ---- | ----------- |
+| `CreateProject` |  |
+| `DeleteProject` |  |
+| `UpdateProject` |  |
+| `ViewProject` |  |
+
+## Replicators
+
+| Name | Description |
+| ---- | ----------- |
+| `UpdateReplicator` |  |
+| `ViewReplicator` |  |
+
 ## ServiceAccounts
 
 | Name | Description |
@@ -80,6 +274,39 @@
 | `DeleteAccount` |  |
 | `ViewAccount` |  |
 
+## Sessions
+
+| Name | Description |
+| ---- | ----------- |
+| `CreateSession` |  |
+| `ViewSession` |  |
+
+## Software
+
+| Name | Description |
+| ---- | ----------- |
+| `UploadSoftware` |  |
+| `DownloadSoftware` |  |
+| `DeleteSoftware` |  |
+
+## Streams
+
+| Name | Description |
+| ---- | ----------- |
+| `CreateStream` |  |
+| `UpdateStream` |  |
+| `DeleteStream` |  |
+| `ViewStream` |  |
+| `ViewChanges` |  |
+| `ViewTemplate` |  |
+
+## TestData
+
+| Name | Description |
+| ---- | ----------- |
+| `TestHealthRead` |  |
+| `TestHealthWrite` |  |
+
 ## Accounts
 
 | Name | Description |
@@ -89,6 +316,17 @@
 | `DeleteAccount` |  |
 | `ViewAccount` |  |
 
+## Agents
+
+| Name | Description |
+| ---- | ----------- |
+| `CreateAgent` |  |
+| `CreateWorkstationAgent` |  |
+| `UpdateAgent` |  |
+| `DeleteAgent` |  |
+| `ViewAgent` |  |
+| `ListAgents` |  |
+
 ## Artifacts
 
 | Name | Description |
@@ -113,6 +351,15 @@
 | ---- | ----------- |
 | `AddUpdateDeleteBuildHealthFilter` |  |
 
+## Compute
+
+| Name | Description |
+| ---- | ----------- |
+| `AddComputeTasks` |  |
+| `GetComputeTasks` |  |
+| `UbaCacheRead` |  |
+| `UbaCacheWrite` |  |
+
 ## Devices
 
 | Name | Description |
@@ -131,6 +378,24 @@
 | `RetryJobStep` |  |
 | `ViewJob` |  |
 
+## Leases
+
+| Name | Description |
+| ---- | ----------- |
+| `ViewLeases` |  |
+| `ViewLeaseTasks` |  |
+
+## Logs
+
+| Name | Description |
+| ---- | ----------- |
+| `CreateLog` |  |
+| `UpdateLog` |  |
+| `ViewLog` |  |
+| `WriteLogData` |  |
+| `CreateEvent` |  |
+| `ViewEvent` |  |
+
 ## Notices
 
 | Name | Description |
@@ -145,6 +410,16 @@
 | ---- | ----------- |
 | `CreateSubscription` |  |
 
+## Pools
+
+| Name | Description |
+| ---- | ----------- |
+| `CreatePool` |  |
+| `UpdatePool` |  |
+| `DeletePool` |  |
+| `ViewPool` |  |
+| `ListPools` |  |
+
 ## Projects
 
 | Name | Description |
@@ -170,6 +445,21 @@
 | `DeleteAccount` |  |
 | `ViewAccount` |  |
 
+## Sessions
+
+| Name | Description |
+| ---- | ----------- |
+| `CreateSession` |  |
+| `ViewSession` |  |
+
+## Software
+
+| Name | Description |
+| ---- | ----------- |
+| `UploadSoftware` |  |
+| `DownloadSoftware` |  |
+| `DeleteSoftware` |  |
+
 ## Streams
 
 | Name | Description |
@@ -197,6 +487,17 @@
 | `DeleteAccount` |  |
 | `ViewAccount` |  |
 
+## Agents
+
+| Name | Description |
+| ---- | ----------- |
+| `CreateAgent` |  |
+| `CreateWorkstationAgent` |  |
+| `UpdateAgent` |  |
+| `DeleteAgent` |  |
+| `ViewAgent` |  |
+| `ListAgents` |  |
+
 ## Artifacts
 
 | Name | Description |
@@ -221,6 +522,15 @@
 | ---- | ----------- |
 | `AddUpdateDeleteBuildHealthFilter` |  |
 
+## Compute
+
+| Name | Description |
+| ---- | ----------- |
+| `AddComputeTasks` |  |
+| `GetComputeTasks` |  |
+| `UbaCacheRead` |  |
+| `UbaCacheWrite` |  |
+
 ## Devices
 
 | Name | Description |
@@ -239,6 +549,24 @@
 | `RetryJobStep` |  |
 | `ViewJob` |  |
 
+## Leases
+
+| Name | Description |
+| ---- | ----------- |
+| `ViewLeases` |  |
+| `ViewLeaseTasks` |  |
+
+## Logs
+
+| Name | Description |
+| ---- | ----------- |
+| `CreateLog` |  |
+| `UpdateLog` |  |
+| `ViewLog` |  |
+| `WriteLogData` |  |
+| `CreateEvent` |  |
+| `ViewEvent` |  |
+
 ## Notices
 
 | Name | Description |
@@ -253,6 +581,16 @@
 | ---- | ----------- |
 | `CreateSubscription` |  |
 
+## Pools
+
+| Name | Description |
+| ---- | ----------- |
+| `CreatePool` |  |
+| `UpdatePool` |  |
+| `DeletePool` |  |
+| `ViewPool` |  |
+| `ListPools` |  |
+
 ## Projects
 
 | Name | Description |
@@ -269,12 +607,6 @@
 | `UpdateReplicator` |  |
 | `ViewReplicator` |  |
 
-## Secrets
-
-| Name | Description |
-| ---- | ----------- |
-| `ViewSecret` |  |
-
 ## ServiceAccounts
 
 | Name | Description |
@@ -284,6 +616,21 @@
 | `DeleteAccount` |  |
 | `ViewAccount` |  |
 
+## Sessions
+
+| Name | Description |
+| ---- | ----------- |
+| `CreateSession` |  |
+| `ViewSession` |  |
+
+## Software
+
+| Name | Description |
+| ---- | ----------- |
+| `UploadSoftware` |  |
+| `DownloadSoftware` |  |
+| `DeleteSoftware` |  |
+
 ## Streams
 
 | Name | Description |
@@ -311,6 +658,17 @@
 | `DeleteAccount` |  |
 | `ViewAccount` |  |
 
+## Agents
+
+| Name | Description |
+| ---- | ----------- |
+| `CreateAgent` |  |
+| `CreateWorkstationAgent` |  |
+| `UpdateAgent` |  |
+| `DeleteAgent` |  |
+| `ViewAgent` |  |
+| `ListAgents` |  |
+
 ## Artifacts
 
 | Name | Description |
@@ -335,6 +693,15 @@
 | ---- | ----------- |
 | `AddUpdateDeleteBuildHealthFilter` |  |
 
+## Compute
+
+| Name | Description |
+| ---- | ----------- |
+| `AddComputeTasks` |  |
+| `GetComputeTasks` |  |
+| `UbaCacheRead` |  |
+| `UbaCacheWrite` |  |
+
 ## Devices
 
 | Name | Description |
@@ -353,6 +720,24 @@
 | `RetryJobStep` |  |
 | `ViewJob` |  |
 
+## Leases
+
+| Name | Description |
+| ---- | ----------- |
+| `ViewLeases` |  |
+| `ViewLeaseTasks` |  |
+
+## Logs
+
+| Name | Description |
+| ---- | ----------- |
+| `CreateLog` |  |
+| `UpdateLog` |  |
+| `ViewLog` |  |
+| `WriteLogData` |  |
+| `CreateEvent` |  |
+| `ViewEvent` |  |
+
 ## Notices
 
 | Name | Description |
@@ -367,6 +752,16 @@
 | ---- | ----------- |
 | `CreateSubscription` |  |
 
+## Pools
+
+| Name | Description |
+| ---- | ----------- |
+| `CreatePool` |  |
+| `UpdatePool` |  |
+| `DeletePool` |  |
+| `ViewPool` |  |
+| `ListPools` |  |
+
 ## Projects
 
 | Name | Description |
@@ -383,12 +778,6 @@
 | `UpdateReplicator` |  |
 | `ViewReplicator` |  |
 
-## Secrets
-
-| Name | Description |
-| ---- | ----------- |
-| `ViewSecret` |  |
-
 ## ServiceAccounts
 
 | Name | Description |
@@ -398,6 +787,21 @@
 | `DeleteAccount` |  |
 | `ViewAccount` |  |
 
+## Sessions
+
+| Name | Description |
+| ---- | ----------- |
+| `CreateSession` |  |
+| `ViewSession` |  |
+
+## Software
+
+| Name | Description |
+| ---- | ----------- |
+| `UploadSoftware` |  |
+| `DownloadSoftware` |  |
+| `DeleteSoftware` |  |
+
 ## Streams
 
 | Name | Description |
@@ -545,12 +949,6 @@
 | `UpdateReplicator` |  |
 | `ViewReplicator` |  |
 
-## Secrets
-
-| Name | Description |
-| ---- | ----------- |
-| `ViewSecret` |  |
-
 ## ServiceAccounts
 
 | Name | Description |
@@ -575,6 +973,19 @@
 | `DownloadSoftware` |  |
 | `DeleteSoftware` |  |
 
+## Storage
+
+| Name | Description |
+| ---- | ----------- |
+| `ReadBlobs` |  |
+| `WriteBlobs` |  |
+| `ReadAliases` |  |
+| `WriteAliases` |  |
+| `DeleteAliases` |  |
+| `ReadRefs` |  |
+| `WriteRefs` |  |
+| `DeleteRefs` |  |
+
 ## Streams
 
 | Name | Description |
@@ -722,12 +1133,6 @@
 | `UpdateReplicator` |  |
 | `ViewReplicator` |  |
 
-## Secrets
-
-| Name | Description |
-| ---- | ----------- |
-| `ViewSecret` |  |
-
 ## ServiceAccounts
 
 | Name | Description |
@@ -752,6 +1157,19 @@
 | `DownloadSoftware` |  |
 | `DeleteSoftware` |  |
 
+## Storage
+
+| Name | Description |
+| ---- | ----------- |
+| `ReadBlobs` |  |
+| `WriteBlobs` |  |
+| `ReadAliases` |  |
+| `WriteAliases` |  |
+| `DeleteAliases` |  |
+| `ReadRefs` |  |
+| `WriteRefs` |  |
+| `DeleteRefs` |  |
+
 ## Streams
 
 | Name | Description |
@@ -763,19 +1181,18 @@
 | `ViewChanges` |  |
 | `ViewTemplate` |  |
 
-## TestData
+## Symbols
 
 | Name | Description |
 | ---- | ----------- |
-| `TestHealthRead` |  |
-| `TestHealthWrite` |  |
+| `ReadSymbols` |  |
 
-## Tools
+## TestData
 
 | Name | Description |
 | ---- | ----------- |
-| `DownloadTool` |  |
-| `UploadTool` |  |
+| `TestHealthRead` |  |
+| `TestHealthWrite` |  |
 
 ## Accounts
 
@@ -936,6 +1353,19 @@
 | `DownloadSoftware` |  |
 | `DeleteSoftware` |  |
 
+## Storage
+
+| Name | Description |
+| ---- | ----------- |
+| `ReadBlobs` |  |
+| `WriteBlobs` |  |
+| `ReadAliases` |  |
+| `WriteAliases` |  |
+| `DeleteAliases` |  |
+| `ReadRefs` |  |
+| `WriteRefs` |  |
+| `DeleteRefs` |  |
+
 ## Streams
 
 | Name | Description |
@@ -960,13 +1390,6 @@
 | `TestHealthRead` |  |
 | `TestHealthWrite` |  |
 
-## Tools
-
-| Name | Description |
-| ---- | ----------- |
-| `DownloadTool` |  |
-| `UploadTool` |  |
-
 ## Accounts
 
 | Name | Description |
diff --git a/Engine/Source/Programs/Horde/Docs/Config/Schema/Globals.md b/Engine/Source/Programs/Horde/Docs/Config/Schema/Globals.md
--- a/Engine/Source/Programs/Horde/Docs/Config/Schema/Globals.md
+++ b/Engine/Source/Programs/Horde/Docs/Config/Schema/Globals.md
@@ -55,27 +55,15 @@
 
 Name | Description
 ---- | -----------
+`build` | [BuildConfig](#buildconfig)<br>Configuration for the build plugin
+`compute` | [ComputeConfig](#computeconfig)<br>Configuration for the compute plugin
 `analytics` | [AnalyticsConfig](#analyticsconfig)<br>Configuration for the analytics plugin
 `exampleplugin` | [ExamplePluginConfig](#examplepluginconfig)<br>Configuration for the exampleplugin plugin
-`build` | [BuildConfig](#buildconfig)<br>Configuration for the build plugin
-`secrets` | [SecretsConfig](#secretsconfig)<br>Configuration for the secrets plugin
 `experimental` | [ExperimentalConfig](#experimentalconfig)<br>Configuration for the experimental plugin
-`compute` | [ComputeConfig](#computeconfig)<br>Configuration for the compute plugin
-`tools` | [ToolsConfig](#toolsconfig)<br>Configuration for the tools plugin
-`symbols` | [SymbolsConfig](#symbolsconfig)<br>Configuration for the symbols plugin
 `storage` | [StorageConfig](#storageconfig)<br>Configuration for the storage plugin
+`symbols` | [SymbolsConfig](#symbolsconfig)<br>Configuration for the symbols plugin
+`secrets` | [SecretsConfig](#secretsconfig)<br>Configuration for the secrets plugin
+`tools` | [ToolsConfig](#toolsconfig)<br>Configuration for the tools plugin
 
-## AnalyticsConfig
-
-Name | Description
----- | -----------
-`stores` | [TelemetryStoreConfig](Telemetry.md)`[]`<br>
-
-## ExamplePluginConfig
-
-Name | Description
----- | -----------
-`greetingMessage` | `string`<br>
-
 ## BuildConfig
 
@@ -106,6 +94,7 @@
 `issueFingerprintFallbackMaxAgeDays` | `integer`<br>
 `issueFingerprintFallbackExcludedTypes` | `string[]`<br>
 `enableExpandedIgnorePatterns` | `boolean`<br>
+`lowContextFingerprintTypes` | `string[]`<br>
 
 ## PerforceCluster
 
@@ -226,78 +215,6 @@
 `excludeActions` | `string[]`<br>
 `extends` | `string[]`<br>
 
-## SecretsConfig
-
-Name | Description
----- | -----------
-`secrets` | [SecretConfig](#secretconfig)`[]`<br>
-`providerConfigs` | [SecretProviderConfig](#secretproviderconfig)`[]`<br>
-
-## SecretConfig
-
-Name | Description
----- | -----------
-`id` | `string`<br>
-`data` | `string` `->` `string`<br>
-`sources` | [ExternalSecretConfig](#externalsecretconfig)`[]`<br>
-`acl` | [AclConfig](#aclconfig)<br>
-
-## ExternalSecretConfig
-
-Name | Description
----- | -----------
-`providerConfig` | `string`<br>
-`provider` | `string`<br>
-`format` | [ExternalSecretFormat](#externalsecretformat-enum)<br>
-`key` | `string`<br>
-`path` | `string`<br>
-
-## ExternalSecretFormat (Enum)
-
-Name | Description
----- | -----------
-`Text` | 
-`Json` | 
-
-## SecretProviderConfig
-
-Name | Description
----- | -----------
-`name` | `string`<br>
-`provider` | `string`<br>
-`hcpVault` | [HcpVaultConfig](#hcpvaultconfig)<br>
-`azureKeyVault` | [AzureKeyVaultConfig](#azurekeyvaultconfig)<br>
-
-## HcpVaultConfig
-
-Name | Description
----- | -----------
-`preSharedKey` | `string`<br>
-`endPoint` | `string`<br>
-`awsArnRole` | `string`<br>
-`awsIamServerId` | `string`<br>
-`role` | `string`<br>
-`credentials` | [HcpVaultCredentialsType](#hcpvaultcredentialstype-enum)<br>
-
-## HcpVaultCredentialsType (Enum)
-
-Name | Description
----- | -----------
-`PreSharedKey` | 
-`AwsAuth` | 
-
-## AzureKeyVaultConfig
-
-Name | Description
----- | -----------
-`vaultUri` | `string`<br>
-
-## ExperimentalConfig
-
-Name | Description
----- | -----------
-`notifications` | [NotificationConfig](Notification.md)`[]`<br>
-
 ## ComputeConfig
 
 Name | Description
@@ -465,44 +382,23 @@
 `description` | `string`<br>
 `computeId` | `string`<br>
 
-## ToolsConfig
+## AnalyticsConfig
 
 Name | Description
 ---- | -----------
-`tools` | [ToolConfig](#toolconfig)`[]`<br>
+`stores` | [TelemetryStoreConfig](Telemetry.md)`[]`<br>
 
-## ToolConfig
+## ExamplePluginConfig
 
 Name | Description
 ---- | -----------
-`id` | `string`<br>
-`name` | `string`<br>
-`description` | `string`<br>
-`category` | `string`<br>
-`group` | `string`<br>
-`platforms` | `string[]`<br>
-`public` | `boolean`<br>
-`showInUgs` | `boolean`<br>
-`showInDashboard` | `boolean`<br>
-`showInToolbox` | `boolean`<br>
-`metadata` | `string` `->` `string`<br>
-`namespaceId` | `string`<br>
-`acl` | [AclConfig](#aclconfig)<br>
+`greetingMessage` | `string`<br>
 
-## SymbolsConfig
+## ExperimentalConfig
 
 Name | Description
 ---- | -----------
-`stores` | [SymbolStoreConfig](#symbolstoreconfig)`[]`<br>
+`notifications` | [NotificationConfig](Notification.md)`[]`<br>
 
-## SymbolStoreConfig
-
-Name | Description
----- | -----------
-`id` | `string`<br>
-`namespaceId` | `string`<br>
-`public` | `boolean`<br>
-`acl` | [AclConfig](#aclconfig)<br>
-
 ## StorageConfig
 
@@ -565,3 +461,108 @@
 `gcDelayHrs` | `number`<br>
 `enableAliases` | `boolean`<br>
 `acl` | [AclConfig](#aclconfig)<br>
+
+## SymbolsConfig
+
+Name | Description
+---- | -----------
+`stores` | [SymbolStoreConfig](#symbolstoreconfig)`[]`<br>
+
+## SymbolStoreConfig
+
+Name | Description
+---- | -----------
+`id` | `string`<br>
+`namespaceId` | `string`<br>
+`public` | `boolean`<br>
+`acl` | [AclConfig](#aclconfig)<br>
+
+## SecretsConfig
+
+Name | Description
+---- | -----------
+`secrets` | [SecretConfig](#secretconfig)`[]`<br>
+`providerConfigs` | [SecretProviderConfig](#secretproviderconfig)`[]`<br>
+
+## SecretConfig
+
+Name | Description
+---- | -----------
+`id` | `string`<br>
+`data` | `string` `->` `string`<br>
+`sources` | [ExternalSecretConfig](#externalsecretconfig)`[]`<br>
+`acl` | [AclConfig](#aclconfig)<br>
+
+## ExternalSecretConfig
+
+Name | Description
+---- | -----------
+`providerConfig` | `string`<br>
+`provider` | `string`<br>
+`format` | [ExternalSecretFormat](#externalsecretformat-enum)<br>
+`key` | `string`<br>
+`path` | `string`<br>
+
+## ExternalSecretFormat (Enum)
+
+Name | Description
+---- | -----------
+`Text` | 
+`Json` | 
+
+## SecretProviderConfig
+
+Name | Description
+---- | -----------
+`name` | `string`<br>
+`provider` | `string`<br>
+`hcpVault` | [HcpVaultConfig](#hcpvaultconfig)<br>
+`azureKeyVault` | [AzureKeyVaultConfig](#azurekeyvaultconfig)<br>
+
+## HcpVaultConfig
+
+Name | Description
+---- | -----------
+`preSharedKey` | `string`<br>
+`endPoint` | `string`<br>
+`awsArnRole` | `string`<br>
+`awsIamServerId` | `string`<br>
+`role` | `string`<br>
+`credentials` | [HcpVaultCredentialsType](#hcpvaultcredentialstype-enum)<br>
+
+## HcpVaultCredentialsType (Enum)
+
+Name | Description
+---- | -----------
+`PreSharedKey` | 
+`AwsAuth` | 
+
+## AzureKeyVaultConfig
+
+Name | Description
+---- | -----------
+`vaultUri` | `string`<br>
+
+## ToolsConfig
+
+Name | Description
+---- | -----------
+`tools` | [ToolConfig](#toolconfig)`[]`<br>
+
+## ToolConfig
+
+Name | Description
+---- | -----------
+`id` | `string`<br>
+`name` | `string`<br>
+`description` | `string`<br>
+`category` | `string`<br>
+`group` | `string`<br>
+`platforms` | `string[]`<br>
+`public` | `boolean`<br>
+`showInUgs` | `boolean`<br>
+`showInDashboard` | `boolean`<br>
+`showInToolbox` | `boolean`<br>
+`metadata` | `string` `->` `string`<br>
+`namespaceId` | `string`<br>
+`acl` | [AclConfig](#aclconfig)<br>
diff --git a/Engine/Source/Programs/Horde/Docs/Deployment/ServerSettings.md b/Engine/Source/Programs/Horde/Docs/Deployment/ServerSettings.md
--- a/Engine/Source/Programs/Horde/Docs/Deployment/ServerSettings.md
+++ b/Engine/Source/Programs/Horde/Docs/Deployment/ServerSettings.md
@@ -125,50 +125,15 @@
 
 Name | Description
 ---- | -----------
+`build` | [BuildServerConfig](#buildserverconfig)<br>Configuration for the build plugin
+`compute` | [ComputeServerConfig](#computeserverconfig)<br>Configuration for the compute plugin
 `analytics` | [AnalyticsServerConfig](#analyticsserverconfig)<br>Configuration for the analytics plugin
 `exampleplugin` | [ExamplePluginServerConfig](#examplepluginserverconfig)<br>Configuration for the exampleplugin plugin
-`build` | [BuildServerConfig](#buildserverconfig)<br>Configuration for the build plugin
-`secrets` | [SecretsServerConfig](#secretsserverconfig)<br>Configuration for the secrets plugin
 `experimental` | [ExperimentalServerConfig](#experimentalserverconfig)<br>Configuration for the experimental plugin
-`compute` | [ComputeServerConfig](#computeserverconfig)<br>Configuration for the compute plugin
-`tools` | [ToolsServerConfig](#toolsserverconfig)<br>Configuration for the tools plugin
-`symbols` | [PluginServerConfig](#pluginserverconfig)<br>Configuration for the symbols plugin
 `storage` | [StorageServerConfig](#storageserverconfig)<br>Configuration for the storage plugin
+`symbols` | [PluginServerConfig](#pluginserverconfig)<br>Configuration for the symbols plugin
+`secrets` | [SecretsServerConfig](#secretsserverconfig)<br>Configuration for the secrets plugin
+`tools` | [ToolsServerConfig](#toolsserverconfig)<br>Configuration for the tools plugin
 
-## AnalyticsServerConfig
-
-Name | Description
----- | -----------
-`sinks` | [TelemetrySinkConfig](#telemetrysinkconfig)<br>
-`enabled` | `boolean`<br>
-
-## TelemetrySinkConfig
-
-Name | Description
----- | -----------
-`epic` | [EpicTelemetryConfig](#epictelemetryconfig)<br>
-`mongo` | [MongoTelemetryConfig](#mongotelemetryconfig)<br>
-
-## EpicTelemetryConfig
-
-Name | Description
----- | -----------
-`url` | `string`<br>
-`appId` | `string`<br>
-`enabled` | `boolean`<br>
-
-## MongoTelemetryConfig
-
-Name | Description
----- | -----------
-`retainDays` | `number`<br>
-`enabled` | `boolean`<br>
-
-## ExamplePluginServerConfig
-
-Name | Description
----- | -----------
-`enabled` | `boolean`<br>
-
 ## BuildServerConfig
 
@@ -306,19 +271,6 @@
 `maxChildCount` | `integer`<br>Maximum number of children in each node
 `sliceThreshold` | `integer`<br>Threshold hash value for splitting interior nodes
 
-## SecretsServerConfig
-
-Name | Description
----- | -----------
-`withAws` | `boolean`<br>
-`enabled` | `boolean`<br>
-
-## ExperimentalServerConfig
-
-Name | Description
----- | -----------
-`enabled` | `boolean`<br>
-
 ## ComputeServerConfig
 
 Name | Description
@@ -361,68 +313,42 @@
 `ComputeQueueAwsMetric` | 
 `LeaseUtilizationAwsMetric` | 
 
-## ToolsServerConfig
+## AnalyticsServerConfig
 
 Name | Description
 ---- | -----------
-`bundledTools` | [BundledToolConfig](#bundledtoolconfig)`[]`<br>
+`sinks` | [TelemetrySinkConfig](#telemetrysinkconfig)<br>
 `enabled` | `boolean`<br>
 
-## BundledToolConfig
+## TelemetrySinkConfig
 
 Name | Description
 ---- | -----------
-`version` | `string`<br>
-`refName` | `string`<br>
-`dataDir` | `string`<br>
-`id` | `string`<br>
-`name` | `string`<br>
-`description` | `string`<br>
-`category` | `string`<br>
-`group` | `string`<br>
-`platforms` | `string[]`<br>
-`public` | `boolean`<br>
-`showInUgs` | `boolean`<br>
-`showInDashboard` | `boolean`<br>
-`showInToolbox` | `boolean`<br>
-`metadata` | `string` `->` `string`<br>
-`namespaceId` | `string`<br>
-`acl` | [AclConfig](#aclconfig)<br>
+`epic` | [EpicTelemetryConfig](#epictelemetryconfig)<br>
+`mongo` | [MongoTelemetryConfig](#mongotelemetryconfig)<br>
 
-## AclConfig
+## EpicTelemetryConfig
 
 Name | Description
 ---- | -----------
-`entries` | [AclEntryConfig](#aclentryconfig)`[]`<br>
-`profiles` | [AclProfileConfig](#aclprofileconfig)`[]`<br>
-`inherit` | `boolean`<br>
-`exceptions` | `string[]`<br>
+`url` | `string`<br>
+`appId` | `string`<br>
+`enabled` | `boolean`<br>
 
-## AclEntryConfig
+## MongoTelemetryConfig
 
 Name | Description
 ---- | -----------
-`claim` | [AclClaimConfig](#aclclaimconfig)<br>
-`actions` | `string[]`<br>
-`profiles` | `string[]`<br>
+`retainDays` | `number`<br>
+`enabled` | `boolean`<br>
 
-## AclClaimConfig
+## ExamplePluginServerConfig
 
 Name | Description
 ---- | -----------
-`type` | `string`<br>
-`value` | `string`<br>
+`enabled` | `boolean`<br>
 
-## AclProfileConfig
+## ExperimentalServerConfig
 
 Name | Description
 ---- | -----------
-`id` | `string`<br>
-`actions` | `string[]`<br>
-`excludeActions` | `string[]`<br>
-`extends` | `string[]`<br>
-
-## PluginServerConfig
-
-Name | Description
----- | -----------
@@ -478,3 +404,77 @@
 `Profile` | 
 `AssumeRole` | 
 `AssumeRoleWebIdentity` | 
+
+## PluginServerConfig
+
+Name | Description
+---- | -----------
+`enabled` | `boolean`<br>
+
+## SecretsServerConfig
+
+Name | Description
+---- | -----------
+`withAws` | `boolean`<br>
+`enabled` | `boolean`<br>
+
+## ToolsServerConfig
+
+Name | Description
+---- | -----------
+`bundledTools` | [BundledToolConfig](#bundledtoolconfig)`[]`<br>
+`enabled` | `boolean`<br>
+
+## BundledToolConfig
+
+Name | Description
+---- | -----------
+`version` | `string`<br>
+`refName` | `string`<br>
+`dataDir` | `string`<br>
+`id` | `string`<br>
+`name` | `string`<br>
+`description` | `string`<br>
+`category` | `string`<br>
+`group` | `string`<br>
+`platforms` | `string[]`<br>
+`public` | `boolean`<br>
+`showInUgs` | `boolean`<br>
+`showInDashboard` | `boolean`<br>
+`showInToolbox` | `boolean`<br>
+`metadata` | `string` `->` `string`<br>
+`namespaceId` | `string`<br>
+`acl` | [AclConfig](#aclconfig)<br>
+
+## AclConfig
+
+Name | Description
+---- | -----------
+`entries` | [AclEntryConfig](#aclentryconfig)`[]`<br>
+`profiles` | [AclProfileConfig](#aclprofileconfig)`[]`<br>
+`inherit` | `boolean`<br>
+`exceptions` | `string[]`<br>
+
+## AclEntryConfig
+
+Name | Description
+---- | -----------
+`claim` | [AclClaimConfig](#aclclaimconfig)<br>
+`actions` | `string[]`<br>
+`profiles` | `string[]`<br>
+
+## AclClaimConfig
+
+Name | Description
+---- | -----------
+`type` | `string`<br>
+`value` | `string`<br>
+
+## AclProfileConfig
+
+Name | Description
+---- | -----------
+`id` | `string`<br>
+`actions` | `string[]`<br>
+`excludeActions` | `string[]`<br>
+`extends` | `string[]`<br>
