Horde: Fix debug test issue summary template — use actual summary text instead of {Summary} placeholder
#rnx

diff --git a/Engine/Source/Programs/Horde/Plugins/Build/HordeServer.Build/Issues/DebugIssuesController.cs b/Engine/Source/Programs/Horde/Plugins/Build/HordeServer.Build/Issues/DebugIssuesController.cs
--- a/Engine/Source/Programs/Horde/Plugins/Build/HordeServer.Build/Issues/DebugIssuesController.cs
+++ b/Engine/Source/Programs/Horde/Plugins/Build/HordeServer.Build/Issues/DebugIssuesController.cs
@@ -116,8 +116,8 @@ namespace HordeServer.Issues
 				issue = await _issueCollection.TryUpdateIssueAsync(issue, null, new UpdateIssueOptions(Promoted: true), HttpContext.RequestAborted) ?? issue;
 			}
 
-			// Build fingerprint
-			NewIssueFingerprint fingerprint = new NewIssueFingerprint(fingerprintType ?? "Compile", "{Summary}", "Code");
+			// Build fingerprint — use the summary text directly as the template so UpdateDerivedDataAsync doesn't overwrite it with "{Summary}"
+			NewIssueFingerprint fingerprint = new NewIssueFingerprint(fingerprintType ?? "Compile", issueSummary, "Code");
 
 			// Use real job/batch/step IDs when provided, otherwise generate synthetic ones
 			JobId effectiveJobId = jobId != null ? JobId.Parse(jobId) : JobIdUtils.GenerateNewId();
