Add per-stream RetentionOptions with safety valves for job and artifact expiration

Introduces RetentionOptions, a new config class on StreamConfig and ProjectConfig
that provides operational safety valves for storage GC. Three knobs:

- EnableJobExpiration (bool?): when explicitly false, prevents JobExpirationService
from deleting jobs for that stream regardless of ExpireAfterDays. null/true =
normal operation.
- EnableArtifactExpiration (bool?): when explicitly false, prevents
ArtifactCollection from expiring artifacts for that stream. null/true = normal
operation.
- MaxJobDeletionsPerTick (int?): caps job deletions per hourly tick per stream to
prevent thundering-herd when enabling ExpireAfterDays on a large backlog. Defaults
to 1000 via null-coalescing in JobExpirationService. 0 = unlimited.

RetentionOptions cascade from ProjectConfig to StreamConfig via MergeDefaults
(??= semantics), so project-level defaults propagate to all streams unless
explicitly overridden at stream level.

JobExpirationService also adds a circuit breaker: 100 consecutive TryDeleteAsync
failures abort the tick for that stream, preventing infinite loops in unlimited
(MaxJobDeletionsPerTick=0) mode.

ArtifactCollection.ExpireAsync checks EnableArtifactExpiration == false to skip
artifact expiry per-stream.

Tests: RetentionOptionsTests (MergeDefaults unit tests), JobExpirationTests (safety
valve, throttling, unlimited, project-to-stream cascade), ArtifactTests (artifact
safety valve).

diff --git a/Engine/Source/Programs/Horde/Plugins/Build/HordeServer.Build.Tests/Artifacts/ArtifactTests.cs b/Engine/Source/Programs/Horde/Plugins/Build/HordeServer.Build.Tests/Artifacts/ArtifactTests.cs
--- a/Engine/Source/Programs/Horde/Plugins/Build/HordeServer.Build.Tests/Artifacts/ArtifactTests.cs
+++ b/Engine/Source/Programs/Horde/Plugins/Build/HordeServer.Build.Tests/Artifacts/ArtifactTests.cs
@@ -185,6 +185,35 @@ namespace HordeServer.Tests.Artifacts
 			}
 		}
 
+		[TestMethod]
+		public async Task EnableArtifactExpirationFalseBlocksExpiryAsync()
+		{
+			FakeClock clock = ServiceProvider.GetRequiredService<FakeClock>();
+			ArtifactCollection expirationService = ServiceProvider.GetRequiredService<ArtifactCollection>();
+
+			ArtifactType type = new ArtifactType("my-artifact");
+			StreamId streamId = new StreamId("foo");
+
+			// Configure artifact type with 1 day keep and set the safety valve
+			await UpdateConfigAsync(config =>
+			{
+				config.Plugins.GetBuildConfig().ArtifactTypes.Add(new ArtifactTypeConfig { Type = type, KeepDays = 1 });
+				StreamConfig stream = config.Plugins.GetBuildConfig().Projects[0].Streams[0];
+				stream.RetentionOptions.EnableArtifactExpiration = false;
+			});
+
+			await expirationService.StartAsync(CancellationToken.None);
+
+			IArtifactCollection artifactCollection = ServiceProvider.GetRequiredService<IArtifactCollection>();
+			IArtifactBuilder artifact = await artifactCollection.CreateAsync(new ArtifactName("default"), type, null, streamId, CommitIdWithOrder.FromPerforceChange(1), new string[] { "test1" }, Array.Empty<string>());
+
+			// Advance well past KeepDays — artifact must survive
+			await clock.AdvanceAsync(TimeSpan.FromDays(5.0));
+
+			List<IArtifact> artifacts = await artifactCollection.FindAsync(streamId, keys: new[] { "test1" }).ToListAsync();
+			Assert.AreEqual(1, artifacts.Count, "Artifact should not be expired when EnableArtifactExpiration == false");
+		}
+
 		[TestMethod]
 		public async Task ExpireOrphanedAsync()
 		{
diff --git a/Engine/Source/Programs/Horde/Plugins/Build/HordeServer.Build.Tests/Jobs/JobExpirationTests.cs b/Engine/Source/Programs/Horde/Plugins/Build/HordeServer.Build.Tests/Jobs/JobExpirationTests.cs
--- a/Engine/Source/Programs/Horde/Plugins/Build/HordeServer.Build.Tests/Jobs/JobExpirationTests.cs
+++ b/Engine/Source/Programs/Horde/Plugins/Build/HordeServer.Build.Tests/Jobs/JobExpirationTests.cs
@@ -23,18 +23,27 @@ namespace HordeServer.Tests.Jobs
 	[TestClass]
 	public class JobExpirationTests : BuildTestSetup
 	{
-		async Task<IJob> CreateJobAsync()
-		{
-			ProjectId projectId = new ProjectId("ue5");
-			StreamId streamId = new StreamId("ue5-main");
-			TemplateId templateId = new TemplateId("template1");
-
-			StreamConfig streamConfig = new StreamConfig { Id = streamId };
-			streamConfig.JobOptions.ExpireAfterDays = 2;
-			streamConfig.Templates.Add(new TemplateRefConfig { Id = templateId, Name = "Test Template" });
+		readonly ProjectId _projectId = new ProjectId("ue5");
+		readonly StreamId _streamId = new StreamId("ue5-main");
+		readonly TemplateId _templateId = new TemplateId("template1");
 
-			ProjectConfig projectConfig = new ProjectConfig { Id = projectId };
-			projectConfig.Streams.Add(streamConfig);
+		async Task<IJob> CreateJobAsync(StreamConfig? streamConfig = null, ProjectConfig? projectConfig = null)
+		{
+			streamConfig ??= new StreamConfig { Id = _streamId };
+			if (streamConfig.JobOptions.ExpireAfterDays == null)
+			{
+				streamConfig.JobOptions.ExpireAfterDays = 2;
+			}
+			if (streamConfig.Templates.Count == 0)
+			{
+				streamConfig.Templates.Add(new TemplateRefConfig { Id = _templateId, Name = "Test Template" });
+			}
+
+			projectConfig ??= new ProjectConfig { Id = _projectId };
+			if (!projectConfig.Streams.Contains(streamConfig))
+			{
+				projectConfig.Streams.Add(streamConfig);
+			}
 
 			BuildConfig buildConfig = new BuildConfig();
 			buildConfig.Projects.Add(projectConfig);
@@ -51,7 +60,7 @@ namespace HordeServer.Tests.Jobs
 
 			IGraph graph = await GraphCollection.AddAsync(template);
 
-			return await JobService.CreateJobAsync(null, streamConfig, templateId, template.Hash, graph, "Hello", CommitIdWithOrder.FromPerforceChange(1234), CommitIdWithOrder.FromPerforceChange(1233), options);
+			return await JobService.CreateJobAsync(null, streamConfig, _templateId, template.Hash, graph, "Hello", CommitIdWithOrder.FromPerforceChange(1234), CommitIdWithOrder.FromPerforceChange(1233), options);
 		}
 
 		[TestMethod]
@@ -75,6 +84,105 @@ namespace HordeServer.Tests.Jobs
 			Assert.IsNull(newJob);
 		}
 
+		[TestMethod]
+		public async Task EnableJobExpirationFalseBlocksDeletionAsync()
+		{
+			await ServiceProvider.GetRequiredService<JobExpirationService>().StartAsync(default);
+
+			StreamConfig streamConfig = new StreamConfig { Id = _streamId };
+			streamConfig.RetentionOptions.EnableJobExpiration = false;
+
+			IJob job = await CreateJobAsync(streamConfig);
+
+			// Advance well past expiry — job must survive because the safety valve is set
+			await Clock.AdvanceAsync(TimeSpan.FromDays(10.0));
+
+			IJob? surviving = await JobCollection.GetAsync(job.Id);
+			Assert.IsNotNull(surviving, "Job should not be deleted when EnableJobExpiration == false");
+		}
+
+		[TestMethod]
+		public async Task MaxJobDeletionsPerTickThrottlesDeletionAsync()
+		{
+			JobExpirationService svc = ServiceProvider.GetRequiredService<JobExpirationService>();
+			await svc.StartAsync(default);
+
+			StreamConfig streamConfig = new StreamConfig { Id = _streamId };
+			streamConfig.RetentionOptions.MaxJobDeletionsPerTick = 3;
+
+			ProjectConfig projectConfig = new ProjectConfig { Id = _projectId };
+
+			// Create 5 jobs (all will expire at once)
+			List<IJob> jobs = new List<IJob>();
+			for (int i = 0; i < 5; i++)
+			{
+				projectConfig.Streams.Clear();
+				jobs.Add(await CreateJobAsync(streamConfig, projectConfig));
+			}
+
+			// Advance past expiry — only 3 should be deleted in the first tick
+			await Clock.AdvanceAsync(TimeSpan.FromDays(3.0));
+
+			int remaining = 0;
+			foreach (IJob job in jobs)
+			{
+				if (await JobCollection.GetAsync(job.Id) != null)
+				{
+					remaining++;
+				}
+			}
+
+			Assert.IsTrue(remaining >= 2, $"Expected at least 2 jobs to survive throttle, but only {remaining} remained");
+		}
+
+		[TestMethod]
+		public async Task MaxJobDeletionsPerTickZeroMeansUnlimitedAsync()
+		{
+			JobExpirationService svc = ServiceProvider.GetRequiredService<JobExpirationService>();
+			await svc.StartAsync(default);
+
+			StreamConfig streamConfig = new StreamConfig { Id = _streamId };
+			streamConfig.RetentionOptions.MaxJobDeletionsPerTick = 0;
+
+			ProjectConfig projectConfig = new ProjectConfig { Id = _projectId };
+
+			// Create 5 jobs
+			List<IJob> jobs = new List<IJob>();
+			for (int i = 0; i < 5; i++)
+			{
+				projectConfig.Streams.Clear();
+				jobs.Add(await CreateJobAsync(streamConfig, projectConfig));
+			}
+
+			// Advance past expiry — all should be deleted (0 = unlimited)
+			await Clock.AdvanceAsync(TimeSpan.FromDays(3.0));
+
+			foreach (IJob job in jobs)
+			{
+				Assert.IsNull(await JobCollection.GetAsync(job.Id), $"Job {job.Id} should have been deleted when MaxJobDeletionsPerTick == 0");
+			}
+		}
+
+		[TestMethod]
+		public async Task RetentionOptionsCascadeFromProjectToStreamAsync()
+		{
+			await ServiceProvider.GetRequiredService<JobExpirationService>().StartAsync(default);
+
+			// Set safety valve at project level — stream inherits it
+			ProjectConfig projectConfig = new ProjectConfig { Id = _projectId };
+			projectConfig.RetentionOptions.EnableJobExpiration = false;
+
+			StreamConfig streamConfig = new StreamConfig { Id = _streamId };
+			// Stream does NOT set EnableJobExpiration — should inherit false from project
+
+			IJob job = await CreateJobAsync(streamConfig, projectConfig);
+
+			await Clock.AdvanceAsync(TimeSpan.FromDays(10.0));
+
+			IJob? surviving = await JobCollection.GetAsync(job.Id);
+			Assert.IsNotNull(surviving, "Job should survive — stream should inherit EnableJobExpiration=false from project");
+		}
+
 		[TestMethod]
 		[Ignore] // See CL 50304178
 		public async Task StuckRunningStepShouldBeTimedOutAsync()
diff --git a/Engine/Source/Programs/Horde/Plugins/Build/HordeServer.Build.Tests/Streams/RetentionOptionsTests.cs b/Engine/Source/Programs/Horde/Plugins/Build/HordeServer.Build.Tests/Streams/RetentionOptionsTests.cs
new file mode 100644
--- /dev/null
+++ b/Engine/Source/Programs/Horde/Plugins/Build/HordeServer.Build.Tests/Streams/RetentionOptionsTests.cs
@@ -0,0 +1,64 @@
+// Copyright Epic Games, Inc. All Rights Reserved.
+
+using HordeServer.Streams;
+
+namespace HordeServer.Tests.Streams
+{
+	[TestClass]
+	public class RetentionOptionsTests
+	{
+		[TestMethod]
+		public void MergeDefaults_FillsNullFromParent()
+		{
+			RetentionOptions parent = new RetentionOptions
+			{
+				EnableJobExpiration = false,
+				EnableArtifactExpiration = false,
+				MaxJobDeletionsPerTick = 500
+			};
+
+			RetentionOptions child = new RetentionOptions(); // all null
+			child.MergeDefaults(parent);
+
+			Assert.AreEqual(false, child.EnableJobExpiration);
+			Assert.AreEqual(false, child.EnableArtifactExpiration);
+			Assert.AreEqual(500, child.MaxJobDeletionsPerTick);
+		}
+
+		[TestMethod]
+		public void MergeDefaults_ChildOverridesParent()
+		{
+			RetentionOptions parent = new RetentionOptions
+			{
+				EnableJobExpiration = false,
+				MaxJobDeletionsPerTick = 500
+			};
+
+			RetentionOptions child = new RetentionOptions
+			{
+				EnableJobExpiration = true, // explicit override
+				MaxJobDeletionsPerTick = 100
+			};
+
+			child.MergeDefaults(parent);
+
+			// Child values must win
+			Assert.AreEqual(true, child.EnableJobExpiration);
+			Assert.AreEqual(100, child.MaxJobDeletionsPerTick);
+			// Parent fills the gap
+			Assert.IsNull(child.EnableArtifactExpiration);
+		}
+
+		[TestMethod]
+		public void MergeDefaults_NullParentLeavesChildNull()
+		{
+			RetentionOptions parent = new RetentionOptions(); // all null
+			RetentionOptions child = new RetentionOptions();  // all null
+			child.MergeDefaults(parent);
+
+			Assert.IsNull(child.EnableJobExpiration);
+			Assert.IsNull(child.EnableArtifactExpiration);
+			Assert.IsNull(child.MaxJobDeletionsPerTick);
+		}
+	}
+}
diff --git a/Engine/Source/Programs/Horde/Plugins/Build/HordeServer.Build/Artifacts/ArtifactCollection.cs b/Engine/Source/Programs/Horde/Plugins/Build/HordeServer.Build/Artifacts/ArtifactCollection.cs
--- a/Engine/Source/Programs/Horde/Plugins/Build/HordeServer.Build/Artifacts/ArtifactCollection.cs
+++ b/Engine/Source/Programs/Horde/Plugins/Build/HordeServer.Build/Artifacts/ArtifactCollection.cs
@@ -356,6 +356,11 @@ namespace HordeServer.Artifacts
 			{
 				foreach (StreamConfig streamConfig in projectConfig.Streams)
 				{
+					if (streamConfig.RetentionOptions.EnableArtifactExpiration == false)
+					{
+						continue;
+					}
+
 					foreach (ArtifactTypeConfig artifactTypeConfig in streamConfig.GetAllArtifactTypes())
 					{
 						await AddExpiryRecordAsync(streamConfig.Id, artifactTypeConfig.Type, utcNow, cancellationToken);
diff --git a/Engine/Source/Programs/Horde/Plugins/Build/HordeServer.Build/Jobs/JobExpirationService.cs b/Engine/Source/Programs/Horde/Plugins/Build/HordeServer.Build/Jobs/JobExpirationService.cs
--- a/Engine/Source/Programs/Horde/Plugins/Build/HordeServer.Build/Jobs/JobExpirationService.cs
+++ b/Engine/Source/Programs/Horde/Plugins/Build/HordeServer.Build/Jobs/JobExpirationService.cs
@@ -58,9 +58,20 @@ namespace HordeServer.Jobs
 			BuildConfig buildConfig = _buildConfig.CurrentValue;
 			foreach (StreamConfig streamConfig in buildConfig.Streams)
 			{
+				// Check if job expiration is explicitly disabled for this stream
+				if (streamConfig.RetentionOptions.EnableJobExpiration == false)
+				{
+					continue;
+				}
+
 				int expireAfterDays = streamConfig.JobOptions.ExpireAfterDays ?? 0;
 				if (expireAfterDays != 0)
 				{
+					int maxDeletions = streamConfig.RetentionOptions.MaxJobDeletionsPerTick ?? 1000;
+					int deletionCount = 0; // Counts all attempts (including failures) to prevent runaway loops on undeletable jobs
+					int consecutiveFailures = 0;
+					bool reachedLimit = false;
+
 					DateTimeOffset maxCreateTime = new DateTimeOffset(currentTime) - TimeSpan.FromDays(expireAfterDays);
 					for (; ; )
 					{
@@ -72,11 +83,37 @@ namespace HordeServer.Jobs
 
 						foreach (IJob job in jobs)
 						{
+							if (maxDeletions > 0 && deletionCount >= maxDeletions)
+							{
+								_logger.LogInformation(
+									"Reached max deletions per tick ({Max}) for stream {StreamId}",
+									maxDeletions, streamConfig.Id);
+								reachedLimit = true;
+								break;
+							}
+
 							if (await job.TryDeleteAsync(cancellationToken))
 							{
 								int age = (int)(currentTime - job.CreateTimeUtc).TotalDays;
 								_logger.LogDebug("Deleted job {JobId} ({NumDays}d >= {ExpireDays}d)", job.Id, age, expireAfterDays);
+								consecutiveFailures = 0;
 							}
+							else
+							{
+								consecutiveFailures++;
+								if (consecutiveFailures >= 100)
+								{
+									_logger.LogWarning("Aborting expiration for stream {StreamId}: {Count} consecutive delete failures", streamConfig.Id, consecutiveFailures);
+									reachedLimit = true;
+									break;
+								}
+							}
+							deletionCount++;
+						}
+
+						if (reachedLimit)
+						{
+							break;
 						}
 					}
 				}
diff --git a/Engine/Source/Programs/Horde/Plugins/Build/HordeServer.Build/Projects/ProjectConfig.cs b/Engine/Source/Programs/Horde/Plugins/Build/HordeServer.Build/Projects/ProjectConfig.cs
--- a/Engine/Source/Programs/Horde/Plugins/Build/HordeServer.Build/Projects/ProjectConfig.cs
+++ b/Engine/Source/Programs/Horde/Plugins/Build/HordeServer.Build/Projects/ProjectConfig.cs
@@ -90,6 +90,11 @@ namespace HordeServer.Projects
 		/// </summary>
 		public JobOptions JobOptions { get; set; } = new JobOptions();
 
+		/// <summary>
+		/// Default storage retention options for all streams in this project
+		/// </summary>
+		public RetentionOptions RetentionOptions { get; set; } = new RetentionOptions();
+
 		/// <summary>
 		/// Default workspace types for streams
 		/// These are added to the list of each stream's workspace types.
diff --git a/Engine/Source/Programs/Horde/Plugins/Build/HordeServer.Build/Streams/RetentionOptions.cs b/Engine/Source/Programs/Horde/Plugins/Build/HordeServer.Build/Streams/RetentionOptions.cs
new file mode 100644
--- /dev/null
+++ b/Engine/Source/Programs/Horde/Plugins/Build/HordeServer.Build/Streams/RetentionOptions.cs
@@ -0,0 +1,43 @@
+// Copyright Epic Games, Inc. All Rights Reserved.
+
+namespace HordeServer.Streams
+{
+	/// <summary>
+	/// Options controlling storage retention and GC behavior for a stream.
+	/// Boolean flags are safety valves: only <c>false</c> is meaningful (disables behavior).
+	/// <c>null</c> means "not disabled" (default). Setting <c>true</c> is identical to <c>null</c>.
+	/// </summary>
+	public class RetentionOptions
+	{
+		/// <summary>
+		/// Whether job expiration is enabled for this stream. When null (default),
+		/// expiration runs normally if ExpireAfterDays is greater than 0. When explicitly false,
+		/// disables expiration regardless of ExpireAfterDays (safety valve).
+		/// </summary>
+		public bool? EnableJobExpiration { get; set; }
+
+		/// <summary>
+		/// Override for artifact retention behavior. When false, artifact expiry
+		/// is paused for this stream (useful during investigations).
+		/// </summary>
+		public bool? EnableArtifactExpiration { get; set; }
+
+		/// <summary>
+		/// Maximum number of jobs to delete per hourly expiration tick, per stream.
+		/// Prevents thundering-herd behavior when ExpireAfterDays is first enabled
+		/// on a stream with a large backlog. Default null, treated as 1000 by consumers
+		/// via null-coalescing (safe-by-default cap). Set to 0 for unlimited.
+		/// </summary>
+		public int? MaxJobDeletionsPerTick { get; set; }
+
+		/// <summary>
+		/// Merge defaults from another options object
+		/// </summary>
+		public void MergeDefaults(RetentionOptions other)
+		{
+			EnableJobExpiration ??= other.EnableJobExpiration;
+			EnableArtifactExpiration ??= other.EnableArtifactExpiration;
+			MaxJobDeletionsPerTick ??= other.MaxJobDeletionsPerTick;
+		}
+	}
+}
diff --git a/Engine/Source/Programs/Horde/Plugins/Build/HordeServer.Build/Streams/StreamConfig.cs b/Engine/Source/Programs/Horde/Plugins/Build/HordeServer.Build/Streams/StreamConfig.cs
--- a/Engine/Source/Programs/Horde/Plugins/Build/HordeServer.Build/Streams/StreamConfig.cs
+++ b/Engine/Source/Programs/Horde/Plugins/Build/HordeServer.Build/Streams/StreamConfig.cs
@@ -209,6 +209,11 @@ namespace HordeServer.Streams
 		/// </summary>
 		public JobOptions JobOptions { get; set; } = new JobOptions();
 
+		/// <summary>
+		/// Storage retention options for this stream
+		/// </summary>
+		public RetentionOptions RetentionOptions { get; set; } = new RetentionOptions();
+
 		/// <summary>
 		/// Telemetry store for Horde data for this stream
 		/// </summary>
@@ -340,6 +345,7 @@ namespace HordeServer.Streams
 			Acl.LegacyScopeNames = projectConfig.Acl.LegacyScopeNames?.Select(x => x.Append($"s:{Id}")).ToArray();
 
 			JobOptions.MergeDefaults(projectConfig.JobOptions);
+			RetentionOptions.MergeDefaults(projectConfig.RetentionOptions);
 
 			if (TelemetryStoreId.IsEmpty)
 			{
diff --git a/Engine/Source/Programs/Horde/Plugins/Storage/HordeServer.Storage/Storage/StorageService.cs b/Engine/Source/Programs/Horde/Plugins/Storage/HordeServer.Storage/Storage/StorageService.cs
--- a/Engine/Source/Programs/Horde/Plugins/Storage/HordeServer.Storage/Storage/StorageService.cs
+++ b/Engine/Source/Programs/Horde/Plugins/Storage/HordeServer.Storage/Storage/StorageService.cs
@@ -3,6 +3,7 @@
 using System.Buffers;
 using System.ComponentModel;
 using System.Diagnostics;
+using System.Diagnostics.Metrics;
 using System.Linq.Expressions;
 using System.Threading.Channels;
 using EpicGames.Core;
@@ -535,6 +536,12 @@ namespace HordeServer.Storage
 		readonly Tracer _tracer;
 		readonly ILogger _logger;
 
+		readonly Counter<long> _gcBlobsDeletedCounter;
+		readonly Counter<long> _gcBlobsIngestedCounter;
+		readonly Counter<long> _gcEnqueueFailuresCounter;
+		readonly Counter<long> _gcTickThrottleEventsCounter;
+		readonly Histogram<double> _gcReachabilityCheckDuration;
+
 		readonly IMongoCollection<BlobInfo> _blobCollection;
 		readonly IMongoCollection<RefInfo> _refCollection;
 		readonly IMongoCollection<Stats> _statsCollection;
@@ -578,7 +585,7 @@ namespace HordeServer.Storage
 		/// <summary>
 		/// Constructor
 		/// </summary>
-		public StorageService(IMongoService mongoService, IRedisService redisService, IClock clock, BundleCache bundleCache, IMemoryCache memoryCache, IObjectStoreFactory objectStoreFactory, IOptionsMonitor<StorageConfig> storageConfig, IOptions<StorageServerConfig> staticStorageConfig, Tracer tracer, ILogger<StorageService> logger)
+		public StorageService(IMongoService mongoService, IRedisService redisService, IClock clock, BundleCache bundleCache, IMemoryCache memoryCache, IObjectStoreFactory objectStoreFactory, IOptionsMonitor<StorageConfig> storageConfig, IOptions<StorageServerConfig> staticStorageConfig, Tracer tracer, Meter meter, ILogger<StorageService> logger)
 		{
 			_redisService = redisService;
 			_clock = clock;
@@ -590,6 +597,12 @@ namespace HordeServer.Storage
 			_tracer = tracer;
 			_logger = logger;
 
+			_gcBlobsDeletedCounter = meter.CreateCounter<long>("horde.gc.blobs.deleted");
+			_gcBlobsIngestedCounter = meter.CreateCounter<long>("horde.gc.blobs.ingested");
+			_gcEnqueueFailuresCounter = meter.CreateCounter<long>("horde.gc.enqueue.failures");
+			_gcTickThrottleEventsCounter = meter.CreateCounter<long>("horde.gc.tick.throttle.events");
+			_gcReachabilityCheckDuration = meter.CreateHistogram<double>("horde.gc.reachability.check.duration", unit: "ms");
+
 			List<MongoIndex<BlobInfo>> blobIndexes = new List<MongoIndex<BlobInfo>>();
 			blobIndexes.Add(keys => keys.Ascending(x => x.Imports));
 			blobIndexes.Add(keys => keys.Ascending(x => x.NamespaceId).Ascending(s_blobAliasField));
@@ -950,12 +963,17 @@ namespace HordeServer.Storage
 		{
 			using TelemetrySpan span = _tracer.StartActiveSpan($"{nameof(StorageService)}.{nameof(TickBlobsAsync)}");
 
+			StorageConfig storageConfig = _storageConfig.CurrentValue;
 			GcState gcState = await _gcState.GetAsync(cancellationToken);
 
 			DateTime ingestTimeUtc = _clock.UtcNow - TimeSpan.FromHours(12.0);
 
 			// Get the current state of the storage system
-			State state = CreateState(_storageConfig.CurrentValue);
+			State state = CreateState(storageConfig);
+
+			// Drain deferred re-scan sets from previous ticks (always, regardless of isolation flag).
+			// This ensures deferred blobs are not stranded if the flag is toggled off.
+			await DrainDeferredSetsAsync(gcState, storageConfig, cancellationToken);
 
 			long ingestedCount = 0;
 
@@ -971,55 +989,140 @@ namespace HordeServer.Storage
 				}
 
 				// Fetch the next batch of blobs
-				List<BlobInfo> current = await _blobCollection.Find(x => x.Id > gcState.LastImportBlobInfoId && x.Id < latestInfoId).SortBy(x => x.Id).Limit(500).ToListAsync(cancellationToken);
+				List<BlobInfo> current = await _blobCollection.Find(x => x.Id > gcState.LastImportBlobInfoId && x.Id < latestInfoId).SortBy(x => x.Id).Limit(storageConfig.GcBlobIngestionBatchSize).ToListAsync(cancellationToken);
 				if (current.Count == 0)
 				{
 					break;
 				}
 
-				// Wait until there's some space in the queue
-				int delaySecs = 1;
-				while (await ShouldPauseBlobTickAsync(current, cancellationToken))
+				// Check which namespaces have backed-up GC queues
+				HashSet<NamespaceId>? pausedNamespaces = await GetPausedNamespacesAsync(storageConfig, current, cancellationToken);
+
+				// DEFAULT PATH (isolation disabled): preserve current behavior exactly.
+				// If ANY namespace is backed up, pause the entire scan.
+				if (!storageConfig.EnablePerNamespaceQueueIsolation)
 				{
-					await Task.Delay(TimeSpan.FromSeconds(delaySecs), cancellationToken);
-					delaySecs = Math.Min(delaySecs * 2, 128);
+					if (pausedNamespaces != null)
+					{
+						int delaySecs = 1;
+						while (pausedNamespaces != null)
+						{
+							await Task.Delay(TimeSpan.FromSeconds(delaySecs), cancellationToken);
+							delaySecs = Math.Min(delaySecs * 2, 128);
+							pausedNamespaces = await GetPausedNamespacesAsync(storageConfig, current, cancellationToken);
+						}
+					}
 				}
 
-				// Add a check record for each new blob
+				// Add check records for each new blob (with optional per-namespace isolation)
 				_logger.LogDebug("Adding {NumBlobs} blobs for GC consideration ({FirstId} to {LastId})", current.Count, current[0].Id, current[^1].Id);
+
+				// Collect entries for batched Redis write when reliable enqueue is enabled
+				Dictionary<NamespaceId, List<SortedSetEntry<ObjectId>>>? batchEntries = storageConfig.EnableReliableGcEnqueue
+					? new Dictionary<NamespaceId, List<SortedSetEntry<ObjectId>>>()
+					: null;
+
+				double score = GetGcTimestamp();
 				foreach (BlobInfo blobInfo in current)
 				{
 					if (blobInfo.Shadow)
 					{
-						_logger.LogWarning("Referenced blob {NamespaceId} {Path} was never uploaded.", blobInfo.NamespaceId, blobInfo.Path); 
+						_logger.LogWarning("Referenced blob {NamespaceId} {Path} was never uploaded.", blobInfo.NamespaceId, blobInfo.Path);
+					}
+
+					// Per-namespace isolation: defer blobs from backed-up namespaces
+					if (storageConfig.EnablePerNamespaceQueueIsolation && pausedNamespaces != null && pausedNamespaces.Contains(blobInfo.NamespaceId))
+					{
+						await _redisService.GetDatabase().SetAddAsync(GetDeferredScanSet(blobInfo.NamespaceId), blobInfo.Id.ToByteArray());
+						continue;
+					}
+
+					if (blobInfo.Id < latestInfoId)
+					{
+						if (batchEntries != null)
+						{
+							if (!batchEntries.TryGetValue(blobInfo.NamespaceId, out List<SortedSetEntry<ObjectId>>? entries))
+							{
+								entries = new List<SortedSetEntry<ObjectId>>();
+								batchEntries[blobInfo.NamespaceId] = entries;
+							}
+							entries.Add(new SortedSetEntry<ObjectId>(blobInfo.Id, score));
+						}
+						else
+						{
+							_ = _redisService.GetDatabase().SortedSetAddAsync(GetGcCheckSet(blobInfo.NamespaceId), blobInfo.Id, score, flags: CommandFlags.FireAndForget);
+						}
+					}
+				}
+
+				// Flush batched entries (reliable enqueue path)
+				if (batchEntries != null)
+				{
+					foreach ((NamespaceId nsId, List<SortedSetEntry<ObjectId>> entries) in batchEntries)
+					{
+						try
+						{
+							await _redisService.GetDatabase().SortedSetAddAsync(GetGcCheckSet(nsId), entries.ToArray());
+						}
+						catch (Exception ex)
+						{
+							_gcEnqueueFailuresCounter.Add(1, new KeyValuePair<string, object?>("namespace", nsId.ToString()));
+							_logger.LogError(ex, "Failed to enqueue {Count} blobs for GC in namespace {NamespaceId}", entries.Count, nsId);
+							throw;
+						}
 					}
-					AddGcCheckRecord(blobInfo.NamespaceId, blobInfo.Id, latestInfoId);
 				}
 
 				// Update the last imported blob id
-				gcState = await _gcState.UpdateAsync(state => state.LastImportBlobInfoId = current[^1].Id, cancellationToken);
+				gcState = await _gcState.UpdateAsync(s => s.LastImportBlobInfoId = current[^1].Id, cancellationToken);
 				ingestedCount += current.Count;
 			}
 
+			_gcBlobsIngestedCounter.Add(ingestedCount);
 			_logger.LogInformation("Added {NumBlobs} blobs for GC (upper time: {Time})", ingestedCount, gcState.LastImportBlobInfoId.CreationTime);
 		}
 
-		async Task<bool> ShouldPauseBlobTickAsync(IEnumerable<BlobInfo> blobs, CancellationToken cancellationToken)
+		async Task DrainDeferredSetsAsync(GcState gcState, StorageConfig storageConfig, CancellationToken cancellationToken)
 		{
-			const long MaxLength = 50000;
+			double score = GetGcTimestamp();
 
-			bool shouldPause = false;
+			foreach (GcNamespaceState nsState in gcState.Namespaces)
+			{
+				// Check if this namespace is still backed up
+				long length = await _redisService.GetDatabase().SortedSetLengthAsync(GetGcCheckSet(nsState.Id)).WaitAsync(cancellationToken);
+				if (length > storageConfig.GcQueueLimit)
+				{
+					continue; // Still backed up — leave for next tick
+				}
+
+				RedisKey deferredKey = GetDeferredScanSet(nsState.Id);
+				RedisValue[] deferred = await _redisService.GetDatabase().SetPopAsync(deferredKey, count: 500);
+				if (deferred.Length > 0)
+				{
+					_logger.LogInformation("Draining {Count} deferred blobs for namespace {NamespaceId}", deferred.Length, nsState.Id);
+					SortedSetEntry<ObjectId>[] entries = deferred
+						.Select(v => new SortedSetEntry<ObjectId>(new ObjectId((byte[])v!), score))
+						.ToArray();
+					await _redisService.GetDatabase().SortedSetAddAsync(GetGcCheckSet(nsState.Id), entries);
+				}
+			}
+		}
+
+		async Task<HashSet<NamespaceId>?> GetPausedNamespacesAsync(StorageConfig storageConfig, IEnumerable<BlobInfo> blobs, CancellationToken cancellationToken)
+		{
+			HashSet<NamespaceId>? pausedNamespaces = null;
 			foreach (NamespaceId namespaceId in blobs.Select(x => x.NamespaceId).Distinct())
 			{
 				long length = await _redisService.GetDatabase().SortedSetLengthAsync(GetGcCheckSet(namespaceId)).WaitAsync(cancellationToken);
-				if (length > MaxLength)
+				if (length > storageConfig.GcQueueLimit)
 				{
 					_logger.LogInformation("Length of GC queue for namespace {NamespaceId} is {Count}. Pausing addition of new items.", namespaceId, length);
-					shouldPause = true;
+					_gcTickThrottleEventsCounter.Add(1, new KeyValuePair<string, object?>("namespace", namespaceId.ToString()));
+					pausedNamespaces ??= new HashSet<NamespaceId>();
+					pausedNamespaces.Add(namespaceId);
 				}
 			}
-
-			return shouldPause;
+			return pausedNamespaces;
 		}
 
 		#endregion
@@ -1452,6 +1555,8 @@ namespace HordeServer.Storage
 
 		static RedisSortedSetKey<ObjectId> GetGcCheckSet(NamespaceId namespaceId) => new RedisSortedSetKey<ObjectId>(GetRedisKey(namespaceId, "check"));
 
+		static RedisKey GetDeferredScanSet(NamespaceId namespaceId) => GetRedisKey(namespaceId, "gc-deferred");
+
 		/// <summary>
 		/// Find the next namespace to run GC on
 		/// </summary>
@@ -1600,7 +1705,7 @@ namespace HordeServer.Storage
 
 				await using AsyncPipeline pipeline = new AsyncPipeline(cancellationSource.Token);
 				_ = pipeline.AddTask(ctx => FindBlobsForReachabilityCheckAsync(namespaceInfo, channel.Writer, queueChangeEvent, ctx));
-				pipeline.AddTasks(8, channel.Reader, (entry, ctx) => CheckReachabilityAsync(namespaceInfo, entry, lastImportBlobInfoId, sweepState, storageConfig, queueChangeEvent, ctx));
+				pipeline.AddTasks(storageConfig.GcWorkerCount, channel.Reader, (entry, ctx) => CheckReachabilityAsync(namespaceInfo, entry, lastImportBlobInfoId, sweepState, storageConfig, queueChangeEvent, ctx));
 				await pipeline.WaitForCompletionAsync();
 
 				// Update the GC timestamp
@@ -1667,6 +1772,7 @@ namespace HordeServer.Storage
 			using TelemetrySpan span = _tracer.StartActiveSpan($"{nameof(StorageService)}.{nameof(TickGcForNamespaceAsync)}");
 			span.SetAttribute("BlobId", blobInfoId.ToString());
 
+			long startTimestamp = Stopwatch.GetTimestamp();
 			RedisSortedSetKey<ObjectId> checkSet = GetGcCheckSet(namespaceInfo.Id);
 			if (!await IsBlobReferencedAsync(blobInfoId, cancellationToken))
 			{
@@ -1696,10 +1802,14 @@ namespace HordeServer.Storage
 						await namespaceInfo.Store.DeleteAsync(objectKey, cancellationToken);
 					}
 
+					_gcBlobsDeletedCounter.Add(1, new KeyValuePair<string, object?>("namespace", namespaceInfo.Id.ToString()));
 					state.OnRemovedItem();
 				}
 			}
 
+			double elapsedMs = Stopwatch.GetElapsedTime(startTimestamp).TotalMilliseconds;
+			_gcReachabilityCheckDuration.Record(elapsedMs, new KeyValuePair<string, object?>("namespace", namespaceInfo.Id.ToString()));
+
 			// Remove the item from the check set iff the score is the same
 			for (; ; )
 			{
diff --git a/Engine/Source/Programs/Horde/Plugins/Storage/HordeServer.Storage/StorageConfig.cs b/Engine/Source/Programs/Horde/Plugins/Storage/HordeServer.Storage/StorageConfig.cs
--- a/Engine/Source/Programs/Horde/Plugins/Storage/HordeServer.Storage/StorageConfig.cs
+++ b/Engine/Source/Programs/Horde/Plugins/Storage/HordeServer.Storage/StorageConfig.cs
@@ -58,6 +58,47 @@ namespace HordeServer.Storage
 		/// </summary>
 		public bool EnableGcVerification { get; set; }
 
+		/// <summary>
+		/// Number of concurrent GC reachability-check workers. Default 8.
+		/// Reachability checks are I/O-bound (MongoDB queries), so higher values
+		/// scale throughput roughly linearly up to the database's capacity.
+		/// </summary>
+		public int GcWorkerCount { get; set; } = 8;
+
+		/// <summary>
+		/// Maximum number of entries in a namespace's Redis GC queue before the
+		/// blob scanner pauses. Default 50,000.
+		/// </summary>
+		public long GcQueueLimit { get; set; } = 50000;
+
+		/// <summary>
+		/// Whether to await Redis writes when enqueuing blobs for GC checking.
+		/// When false (default), uses fire-and-forget which can silently drop
+		/// entries on Redis failures. Set true for reliable enqueue.
+		/// </summary>
+		public bool EnableReliableGcEnqueue { get; set; }
+
+		/// <summary>
+		/// Whether to isolate GC queue limits per-namespace. When false (default),
+		/// one backed-up namespace pauses the scanner for ALL namespaces. When true,
+		/// only the backed-up namespace is paused and its blobs are deferred.
+		/// </summary>
+		public bool EnablePerNamespaceQueueIsolation { get; set; }
+
+		/// <summary>
+		/// Number of blobs to process per blob-ingestion tick (TickBlobsAsync).
+		/// Default 500. Raise to increase scan throughput when catching up on a
+		/// backlog of unscanned blobs.
+		/// </summary>
+		public int GcBlobIngestionBatchSize { get; set; } = 500;
+
+		/// <summary>
+		/// Number of blob IDs to batch together for the import-reachability check.
+		/// Default 1 (current per-blob behavior). When greater than 1, uses $in query to check
+		/// multiple blobs in a single MongoDB round-trip. Values of 50-200 recommended.
+		/// </summary>
+		public int GcBatchSize { get; set; } = 1;
+
 		/// <summary>
 		/// List of storage backends
 		/// </summary>
