Add phase grid view
* Add a new alternate view to allow comparing result outcomes for each phase of a test for each test environment (platforms, build target, variant, etc)

#rnx
#Synced-CL 51412642

diff --git a/Engine/Source/Programs/Horde/HordeDashboard/plugins/test-automation/phaseGrid.ts b/Engine/Source/Programs/Horde/HordeDashboard/plugins/test-automation/phaseGrid.ts
new file mode 100644
--- /dev/null
+++ b/Engine/Source/Programs/Horde/HordeDashboard/plugins/test-automation/phaseGrid.ts
@@ -0,0 +1,525 @@
+// Copyright Epic Games, Inc. All Rights Reserved.
+
+import { ContextualMenuItemType } from "@fluentui/react";
+import { TestDataHandler, TestPhaseStatus, MetadataRef, TestPhaseRef, TestMetaStatus, TestNameRef, defaultQueryWeeks, consecutivePassesThreshold } from "./testData";
+import { CircularList, DomainNameColors, ICheckListOption, IPhaseFilter } from "./testAutomationCommon";
+import { action, makeObservable, observable } from "mobx";
+import { TestPhaseOutcome } from "./api";
+
+export type MetadataItem = { meta: MetadataRef, name: string, orderMarker: string, visible: boolean };
+
+export class PhaseStatusGroup implements IPhaseFilter {
+    constructor(phase: TestPhaseRef) {
+        this.phase = phase;
+        this.metadata = new Map();
+    }
+
+    addMetadata(metaItem: MetadataItem, status: TestPhaseStatus): void {
+        if (!this.hasMetadata(metaItem)) {
+            this.metadata.set(metaItem, status);
+            if (status.errorFingerprint) {
+                this.addErrorFingerprint(status.errorFingerprint, metaItem);
+            }
+        }
+    }
+
+    getMetadataStatus(metaItem: MetadataItem): TestPhaseStatus | undefined {
+        return this.metadata.get(metaItem);
+    }
+
+    hasMetadata(metaItem: MetadataItem): boolean {
+        return this.metadata.has(metaItem);
+    }
+
+    hasOutcome(outcome: TestPhaseOutcome): boolean {
+        return this.metadata.entries().filter(([metaItem,]) => metaItem.visible).some(([,status]) => status.outcome === outcome);
+    }
+
+    hasOutcomes(outcomes: Set<TestPhaseOutcome>): boolean {
+        return this.metadata.entries().filter(([metaItem,]) => metaItem.visible).some(([,status]) => outcomes.has(status.outcome));
+    }
+
+    hasOnlyOutcomes(outcomes: Set<TestPhaseOutcome>): boolean {
+        return this.metadata.entries().filter(([metaItem,]) => metaItem.visible).every(([,status]) => outcomes.has(status.outcome));
+    }
+
+    hasWarnings(): boolean {
+        return this.metadata.entries().filter(([metaItem,]) => metaItem.visible).some(([,status]) => status.outcome === TestPhaseOutcome.Success && status.hasWarning);
+    }
+
+    isMatch(names: string[], outcomes?: Set<TestPhaseOutcome>, tags?: string[]): boolean {
+        return this.metadata.entries().filter(([metaItem,]) => metaItem.visible).some(([,status]) => status.isMatch(names, outcomes, tags));
+    }
+
+    getMetadataWithOutcome(outcome: TestPhaseOutcome): MetadataRef[] {
+        return Array.from(this.metadata.keys()).filter(metaItem => metaItem.visible && this.metadata.get(metaItem)?.outcome === outcome).map(metaItem => metaItem.meta);
+    }
+
+    getMetadataWithoutOutcomes(outcomes: Set<TestPhaseOutcome>): MetadataRef[] {
+        return Array.from(this.metadata.keys()).filter(metaItem => metaItem.visible && !outcomes.has(this.metadata.get(metaItem)?.outcome ?? TestPhaseOutcome.Unknown)).map(metaItem => metaItem.meta);
+    }
+
+    getMetadataWithWarnings(): MetadataRef[] {
+        return Array.from(this.metadata.keys()).filter(metaItem => metaItem.visible && this.metadata.get(metaItem)?.outcome === TestPhaseOutcome.Success && this.metadata.get(metaItem)?.hasWarning).map(metaItem => metaItem.meta);
+    }
+
+    addErrorFingerprint(errorFingerprint: string, metaItem: MetadataItem): void {
+        if (!this.errorFingerprints) {
+            this.errorFingerprints = new Map();
+        }
+        if (!this.errorFingerprints.has(errorFingerprint)) {
+            this.errorFingerprints.set(errorFingerprint, []);
+        }
+        this.errorFingerprints.get(errorFingerprint)!.push(metaItem);
+    }
+
+    phase: TestPhaseRef;
+    metadata: Map<MetadataItem, TestPhaseStatus>;
+    componentRef?: any;
+    errorFingerprints?: Map<string, MetadataItem[]>;
+}
+
+const notFailureOutcomes = new Set([TestPhaseOutcome.Success, TestPhaseOutcome.NotRun, TestPhaseOutcome.Skipped]);
+
+export class MetaSelectorContext {
+    constructor(handler: TestDataHandler, test?: TestNameRef) {
+        this._handler = handler;
+        this.test = test;
+        this.metaValueColors = new DomainNameColors();
+        this._metas = [];
+        this._commonKeys = [];
+        this._phaseGroups = [];
+        this._filteredPhases = [];
+        this.initiateContext();
+        this._phaseCallbacks = new Map();
+        this._filterOptions = this.generateFilterOptions();
+        makeObservable(this);
+    }
+
+    @observable
+    metaVisibility: number = 0;
+    @action
+    updateMetaVisibility() {
+        this.metaVisibility++;
+    }
+    subscribeMetaVisibility() {
+        if (this.metaVisibility) { }
+    }
+
+    @observable
+    phaseGroupsUpdated: number = 0;
+    @action
+    updatePhaseGroups() {
+        this.phaseGroupsUpdated++;
+    }
+    subscribePhaseGroups() {
+        if (this.phaseGroupsUpdated) { }
+    }
+
+    @observable
+    filteredPhasesUpdated: number = 0;
+    @action
+    updateFilteredPhases() {
+        this.filteredPhasesUpdated++;
+    }
+    subscribeFilteredPhases() {
+        if (this.filteredPhasesUpdated) { }
+    }
+
+    @observable
+    selectedMetaUpdated: number = 0;
+    @action
+    updateSelectedMeta() {
+        this.selectedMetaUpdated++;
+    }
+    subscribeSelectedMeta() {
+        if (this.selectedMetaUpdated) { }
+    }
+
+    @observable
+    selectedPhaseUpdated: number = 0;
+    @action
+    updateSelectedPhase() {
+        this.selectedPhaseUpdated++;
+    }
+    subscribeSelectedPhase() {
+        if (this.selectedPhaseUpdated) { }
+    }
+
+    @observable
+    metaFilterUpdated: number = 0;
+    @action
+    updateMetaFilter() {
+        this.metaFilterUpdated++;
+    }
+    subscribeMetaFilter() {
+        if (this.metaFilterUpdated) { }
+    }
+
+    private initiateContext() {
+
+        const filteredMeta = new Set(this._handler.filteredMetadata);
+        const testStatus = MetaSelectorContext.getTestStatus(this._handler, this.test);
+
+        const testMetadata = testStatus?.getMetadata().filter(meta => filteredMeta.has(meta)).map(meta => meta) ?? [];
+        this._commonKeys = MetadataRef.identifyCommonKeys(testMetadata);
+        this.commonMetaString  = testMetadata.at(0)?.getCommonValues(this._commonKeys)?.join(" / ");
+
+        this._metas = testMetadata.map(m => {
+            const name = m.getValuesExcept(this._commonKeys).join(" / ");
+            return ({meta: m, name, orderMarker: name, visible: true});
+        });
+
+        const keys: string[] = [];
+        const seen = new Set<string>();
+        this._metas.forEach(metaItem => {
+            metaItem.meta.orderedKeys.forEach(key => {
+                if (!this._commonKeys.includes(key) && !seen.has(key)) {
+                    seen.add(key);
+                    keys.push(key);
+                }
+            });
+        });
+        this.availableMetaKeys = keys;
+
+        const set = new Set(this._metas.flatMap(metaItem => metaItem.meta.getValuesExcept(this._commonKeys)));
+        this.metaValueColors.setDomain(Array.from(set).sort((a, b) => a.localeCompare(b)));
+
+        this.findForkingMetaKey();
+
+    }
+
+    findForkingMetaKey(): string | undefined {
+        // Select the first meta key from allMetaKeys that gives different value across visible meta items
+        const visibleMetas = this._metas.filter(m => m.visible);
+
+        this.forkingMetaKey = this.availableMetaKeys?.find(key => {
+            if (!key) return false;
+            const firstVal = visibleMetas[0]?.meta.getValue?.(key);
+            return visibleMetas.some(m => m.meta.getValue?.(key) !== firstVal);
+        }) ?? this.availableMetaKeys?.at(-1); // fallback to first key if none differ
+
+        visibleMetas.forEach(m => {
+            m.orderMarker = !!this.forkingMetaKey ? m.meta.getValue(this.forkingMetaKey) ?? "" : m.name;
+        });
+        this._metas.sort((a, b) => a.orderMarker.localeCompare(b.orderMarker));    
+
+        this.metaKeysBeforeFork = !!this.forkingMetaKey ? this.availableMetaKeys?.slice(0, this.availableMetaKeys.indexOf(this.forkingMetaKey)) : undefined;
+
+        return this.forkingMetaKey;
+    }
+
+    getForkingMetaKeyValues(): string[] {
+        const forkingKey = this.forkingMetaKey;
+        if (!forkingKey) return [];
+        return this._metas.filter(m => m.visible).map(item => item.meta.getValue(forkingKey)).filter((m, index, self) => self.indexOf(m) === index) as string[];
+    }
+
+    getMetaValuesBeforeForkingKey(): string[] {
+        const forkingKey = this.forkingMetaKey;
+        if (!forkingKey) return [];
+        const visibleMetas = this._metas.filter(m => m.visible);
+        return visibleMetas.map(m => this.metaKeysBeforeFork?.map(key => m.meta.getValue(key))).flat().filter((v, index, self) => !!v && self.indexOf(v) === index) as string[];
+    }
+
+    getMetasContainingValue(value: string): MetadataItem[] {
+        return this._metas.filter(m => m.visible && m.meta.getValuesExcept(this._commonKeys).includes(value));
+    }
+
+    getVisibleMetasContainingForkingKeyValue(value: string): MetadataItem[] {
+        return this._metas.filter(m => m.visible && m.meta.getValue(this.forkingMetaKey!) === value);
+    }
+
+    private generateFilterOptions(): ICheckListOption[] {
+        // Group all possible metadata values by metadata key
+        const groupedByKey: Map<string, Map<string, MetadataItem[]>> = new Map();
+        this._metas.forEach(metaItem => {
+            metaItem.meta.getKeys().forEach(key => {
+                if (this._commonKeys.includes(key)) return;
+                const value = metaItem.meta.getValue?.(key);
+                if (value === undefined) return; // Skip keys that do not yield a value
+                if (!groupedByKey.has(key)) groupedByKey.set(key, new Map());
+                const group = groupedByKey.get(key)!;
+                if (!group.has(value)) group.set(value, []);
+                const groupItems = group.get(value)!;
+                groupItems.push(metaItem);
+            });
+        });
+
+        const gridFilterItems = this.getSearchGridFilterItems();
+
+        // Section headers per key, and their meta values/items
+        const options = this.availableMetaKeys!.map(k => ({key: k, values: groupedByKey.get(k)!})).flatMap((group) => [
+            { key: `key-${group.key}`, itemType: ContextualMenuItemType.Header, text: group.key } as ICheckListOption, // Section header
+            ...group.values.entries().toArray().sort((a, b) => a[0].localeCompare(b[0])).map(([value, items]) => {
+                const option = {
+                    key: `value-${value}-${group.key}`,
+                    text: value,
+                    checked: gridFilterItems?.includes(value) ?? false,
+                    color: this.metaValueColors?.get(value),
+                    data: {key: group.key, items: items},
+                } as ICheckListOption;
+                return option;
+            })
+        ]);
+
+        this.updateMetasVisibility(options);
+
+        return options;
+    }
+
+    get filterOptions(): ICheckListOption[] {
+        return this._filterOptions;
+    }
+
+    updateMetasVisibility(items: ICheckListOption[]) {
+        // Create a Map associated by item key
+        const toCheckPerKey = new Map<string, ICheckListOption[]>();
+        items.filter(i => i.itemType !== ContextualMenuItemType.Header && i.checked).forEach(i => {
+            if (!i.data?.key) return;
+            if (!toCheckPerKey.has(i.data.key)) {
+                toCheckPerKey.set(i.data.key, []);
+            }
+            toCheckPerKey.get(i.data.key)!.push(i);
+        });
+        this._metas.forEach(m => {
+            m.visible = toCheckPerKey.size === 0 || toCheckPerKey.values().every(group => group.some(i => i.data?.items.includes(m)));
+        });
+        this.updateMetaVisibility();
+    }
+
+    get searchUpdated(): number {
+        return this._handler.searchUpdated;
+    }
+
+    get daysQueried(): number {
+        return (this._handler.filterState.weeks ?? defaultQueryWeeks) * 7;
+    }
+
+    getCommonKeysToFork(): string[] {
+        return this._commonKeys.concat(this.metaKeysBeforeFork ?? []);
+    }
+
+    static getTestStatus(handler: TestDataHandler, test?: TestNameRef): TestMetaStatus | undefined {
+        return test && handler.getStatusStream(handler.stream!)?.tests.get(test);
+    }
+
+    getSearchGridFilterItems(): string[] | undefined {
+        const gridFilter = this._handler.getSearchParam('gridfilter') as string | string[] | undefined;
+        return !gridFilter ? undefined : Array.isArray(gridFilter) ? gridFilter : [gridFilter];
+    }
+
+    setSearchGridFilterItems(items: string[], replace?: boolean) {
+        if (items) {
+            this._handler.setSearchParam('gridfilter', items, replace);
+        } else {
+            this._handler.removeSearchParam('gridfilter', replace);
+        }
+    }
+
+    get phaseGroups(): PhaseStatusGroup[] {
+        return this._phaseGroups;
+    }
+
+    setPhaseGroups(groups: PhaseStatusGroup[]) {
+        this._phaseGroups = groups;
+        this.updatePhaseGroups();
+    }
+
+    get filteredPhases(): PhaseStatusGroup[] {
+        return this._filteredPhases;
+    }
+
+    setFilteredPhases(filtered: PhaseStatusGroup[]) {
+        this._filteredPhases = filtered;
+        this.updateFilteredPhases();
+        this.updateFilter(filtered);
+    }
+
+    get selectedMeta(): MetadataItem | undefined {
+        return this._selectedMeta;
+    }
+
+    setSelectedMeta(meta: MetadataItem | undefined) {
+        if (meta === this._selectedMeta) return false;
+        this._selectedMeta = meta;
+        this.updateSelectedMeta();
+        return true;
+    }
+
+    setSelectedMetaById(metaId: string | undefined) {
+        if (!metaId) {
+            return this.setSelectedMeta(undefined);
+        } else if (metaId !== this._selectedMeta?.meta.id) {
+            return this.setSelectedMeta(this._metas.find(m => m.meta.id === metaId));
+        }
+        return false;
+    }
+
+    get selectedPhase(): string | undefined {
+        return this._selectedPhase;
+    }
+
+    setSelectedPhase(phase: string | undefined) {
+        if (phase === this._selectedPhase) return false;
+        this._selectedPhase = phase;
+        this.focusPhase = phase;
+        this.updateSelectedPhase();
+        return true;
+    }
+
+    get failureCursor(): CircularList<string> | undefined {
+        return this._failureCursor;
+    }
+    get interruptedCursor(): CircularList<string> | undefined {
+        return this._interruptedCursor;
+    }
+    get warningCursor(): CircularList<string> | undefined {
+        return this._warningCursor;
+    }
+
+    get metas(): MetadataItem[] {
+        return this._metas;
+    }
+
+    private updateFilter(groups: PhaseStatusGroup[]) {
+        const failures = groups.filter((p) => !p.hasOnlyOutcomes(notFailureOutcomes)).map((p) => p.phase.key) ?? [];
+        this._failureCursor = new CircularList(failures);
+        this._failureCursor.setCursor(this._selectedPhase);
+        const interrupted = groups.filter((p) => p.hasOutcome(TestPhaseOutcome.Interrupted)).map((p) => p.phase.key) ?? [];
+        this._interruptedCursor = new CircularList(interrupted);
+        this._interruptedCursor.setCursor(this._selectedPhase);
+        const warnings = groups.filter((p) => p.hasWarnings()).map((p) => p.phase.key) ?? [];
+        this._warningCursor = new CircularList(warnings);
+        this._warningCursor.setCursor(this._selectedPhase);
+        this._filteredPhases = groups;
+    }
+
+    getNextFailure(): [string | undefined, MetadataItem | undefined] {
+        const nextPhase = this._failureCursor?.next();
+        if (nextPhase) {
+            const nextMetadata = this._filteredPhases.find((g) => g.phase.key === nextPhase)?.getMetadataWithoutOutcomes(notFailureOutcomes).at(0);
+            this._interruptedCursor?.setCursor(nextPhase);
+            return [nextPhase, this._metas.find(m => m.meta === nextMetadata)];
+        }
+        return [undefined, undefined];
+    }
+
+    getPreviousFailure(): [string | undefined, MetadataItem | undefined] {
+        const previousPhase = this._failureCursor?.back();
+        if (previousPhase) {
+            const previousMetadata = this._filteredPhases.find((g) => g.phase.key === previousPhase)?.getMetadataWithoutOutcomes(notFailureOutcomes).at(0);
+            this._interruptedCursor?.setCursor(previousPhase);
+            return [previousPhase, this._metas.find(m => m.meta === previousMetadata)];
+        }
+        return [undefined, undefined];
+    }
+
+    getNextInterrupted(): [string | undefined, MetadataItem | undefined] {
+        const nextPhase = this._interruptedCursor?.next();
+        if (nextPhase) {
+            const nextMetadata = this._filteredPhases.find((g) => g.phase.key === nextPhase)?.getMetadataWithOutcome(TestPhaseOutcome.Interrupted).at(0);
+            this._failureCursor?.setCursor(nextPhase);
+            return [nextPhase, this._metas.find(m => m.meta === nextMetadata)];
+        }
+        return [undefined, undefined];
+    }
+
+    getPreviousInterrupted(): [string | undefined, MetadataItem | undefined] {
+        const previousPhase = this._interruptedCursor?.back();
+        if (previousPhase) {
+            const previousMetadata = this._filteredPhases.find((g) => g.phase.key === previousPhase)?.getMetadataWithOutcome(TestPhaseOutcome.Interrupted).at(0);
+            this._failureCursor?.setCursor(previousPhase);
+            return [previousPhase, this._metas.find(m => m.meta === previousMetadata)];
+        }
+        return [undefined, undefined];
+    }
+
+    getNextWarning(): [string | undefined, MetadataItem | undefined] {
+        const nextPhase = this._warningCursor?.next();
+        if (nextPhase) {
+            const nextMetadata = this._filteredPhases.find((g) => g.phase.key === nextPhase)?.getMetadataWithWarnings().at(0);
+            return [nextPhase, this._metas.find(m => m.meta === nextMetadata)];
+        }
+        return [undefined, undefined];
+    }
+
+    getPreviousWarning(): [string | undefined, MetadataItem | undefined] {
+        const previousPhase = this._warningCursor?.back();
+        if (previousPhase) {
+            const previousMetadata = this._filteredPhases.find((g) => g.phase.key === previousPhase)?.getMetadataWithWarnings().at(0);
+            return [previousPhase, this._metas.find(m => m.meta === previousMetadata)];
+        }
+        return [undefined, undefined];
+    }
+
+    registerOnMetaHovered(onMetaHovered: (meta?: MetadataItem) => void) {
+        this.onMetaHovered = onMetaHovered;
+
+        return () => {
+            this.onMetaHovered = undefined;
+        }
+    }
+
+    registerOnPhaseDataLoaded(phaseKey: string, onDataLoaded: () => void) {
+        const callbacks = this._phaseCallbacks.get(phaseKey) ?? {onDataLoaded: undefined};
+        callbacks.onDataLoaded = onDataLoaded;
+        this._phaseCallbacks.set(phaseKey, callbacks);
+        return () => {
+            callbacks.onDataLoaded = undefined;
+        }
+    }
+
+    notifyOnPhaseDataLoaded(phaseKey: string) {
+        const callback = this._phaseCallbacks.get(phaseKey);
+        if (callback) {
+            callback.onDataLoaded?.();
+        }
+    }
+
+    getFirstTimeErrorsCaught(group: PhaseStatusGroup): Map<string, Date> {
+        const errorsCaught = new Map<string, Date>();
+        const testStatus = MetaSelectorContext.getTestStatus(this._handler, this.test);
+        group.metadata.forEach((status, metaItem) => {
+            if (status?.errorFingerprint) {
+                let startDate = errorsCaught.get(status.errorFingerprint) ?? status.session?.start ?? new Date();
+                const metaStatus = testStatus?.getPhaseSessions(group.phase.key, metaItem.meta);
+                let consecutivePasses = 0;
+                metaStatus?.some(s => {
+                    if (!s.errorFingerprint) {
+                        consecutivePasses++;
+                        return consecutivePasses > consecutivePassesThreshold;
+                    }
+                    if (s.start && s.start < startDate) {
+                        startDate = s.start;
+                    }
+                    consecutivePasses = 0;
+                    return false;
+                });
+                errorsCaught.set(status.errorFingerprint, startDate);
+            }
+        });
+        return errorsCaught;
+    }
+
+    onMetaHovered?: (meta?: MetadataItem) => void;
+    metaValueColors: DomainNameColors;
+    availableMetaKeys?: string[];
+    forkingMetaKey?: string;
+    metaKeysBeforeFork?: string[];
+    test?: TestNameRef;
+    commonMetaString?: string;
+    focusPhase?: string;
+    private _metas: MetadataItem[];
+    private _commonKeys: string[];
+    private _handler: TestDataHandler;
+    private _filterOptions: ICheckListOption[];
+    private _phaseGroups: PhaseStatusGroup[];
+    private _filteredPhases: PhaseStatusGroup[];
+    private _selectedMeta?: MetadataItem;
+    private _selectedPhase?: string;
+    private _failureCursor?: CircularList<string>;
+    private _interruptedCursor?: CircularList<string>;
+    private _warningCursor?: CircularList<string>;
+    private _phaseCallbacks: Map<string, {onDataLoaded?: () => void}>;
+}
diff --git a/Engine/Source/Programs/Horde/HordeDashboard/plugins/test-automation/phaseHistoryGraph.ts b/Engine/Source/Programs/Horde/HordeDashboard/plugins/test-automation/phaseHistoryGraph.ts
--- a/Engine/Source/Programs/Horde/HordeDashboard/plugins/test-automation/phaseHistoryGraph.ts
+++ b/Engine/Source/Programs/Horde/HordeDashboard/plugins/test-automation/phaseHistoryGraph.ts
@@ -1,19 +1,29 @@
 // Copyright Epic Games, Inc. All Rights Reserved.
 
 import * as d3 from "d3";
-import { TestDataHandler, TestNameRef, MetadataRef, PhaseSessionResult, TestMetaStatus } from './testData';
-import dashboard, { StatusColor } from 'horde/backend/Dashboard';
+import { TestDataHandler, TestNameRef, MetadataRef, PhaseSessionResult, TestMetaStatus, consecutivePassesThreshold } from './testData';
+import dashboard from 'horde/backend/Dashboard';
 import { getHumanTime, getShortNiceTime } from 'horde/base/utilities/timeUtils';
 import { getHordeStyling } from 'horde/styles/Styles';
-import { TestOutcome } from "./api";
-import { getPhaseSessionStatusColor, getStatusColors, SessionListSelector } from "./testAutomationCommon";
+import { getPhaseSessionStatusColor,  SessionListSelector } from "./testAutomationCommon";
 import { renderToString } from "react-dom/server";
 
 // Handle bad "@types/d3" types, fix if addressed upstream
 const _d3 = d3 as any;
 
 type SelectionType = d3.Selection<SVGGElement, unknown, null, undefined>;
-type DivSelectionType = d3.Selection<HTMLDivElement, unknown, null, undefined>;
+type DivSelectionType = d3.Selection<HTMLDivElement, unknown, HTMLElement, undefined>;
+type FingerprintSeries = { 
+    fingerprint: string,
+    color: string,
+    count: number,
+    start: number,
+    end: number,
+    y: number
+};
+
+const minimumBandLength = 60 * 60 * 2; // 2 hours
+const minimumBandCount = 2;
 
 export class PhaseHistoryGraph {
 
@@ -34,7 +44,7 @@ export class PhaseHistoryGraph {
         this.sessions = this.status.sessions.get(meta)?.getPhaseSessions(this.phase) ?? [];
     }
 
-    render(container: HTMLDivElement, meta: MetadataRef, selectedPhaseSessionId?: string, withoutHeader?: boolean) {
+    render(container: HTMLDivElement, meta: MetadataRef, selectedPhaseSessionId?: string, withoutHeader?: boolean, reverse?: boolean) {
 
         const { modeColors } = getHordeStyling();
 
@@ -59,12 +69,49 @@ export class PhaseHistoryGraph {
         const yPadding = 1;
         const height = withoutHeader? 20 : Math.ceil((yDomain.size + yPadding) * 13) + this.margin.top + this.margin.bottom;
 
-        const xRange = [this.margin.left, width - this.margin.right];
+        const xRange = reverse? [width - this.margin.right, this.margin.left] : [this.margin.left, width - this.margin.right];
         const yRange = [this.margin.top, height - this.margin.bottom];
 
         const xScale = _d3.scaleTime(xDomain, xRange);
         const yScale = _d3.scalePoint(yDomain, yRange).round(true).padding(yPadding);
 
+        // group sessions into consecutive runs of the same error fingerprint
+        const fingerprintSeries: FingerprintSeries[] = [];
+        let currentSeries: FingerprintSeries | undefined = undefined;
+        let lastFingerprint: string | undefined = undefined;
+        let consecutivePasses = 0;
+
+        sessions.forEach((session, idx) => {
+            const fingerprint = session.errorFingerprint;
+            const x = X[idx];
+            const y = Y[idx];
+
+            if (!fingerprint) {
+                consecutivePasses++;
+                if (consecutivePasses > consecutivePassesThreshold) {
+                    // reset current series if we had more than consecutivePassesThreshold
+                    currentSeries = undefined;
+                    lastFingerprint = undefined;
+                }
+                return;
+            }
+
+            if (fingerprint !== lastFingerprint && (!currentSeries || currentSeries.end !== x)) {
+                // New series started
+                currentSeries = { fingerprint, color: getPhaseSessionStatusColor(session), count: 0, start: x, end: x, y: y };
+                fingerprintSeries.push(currentSeries);
+            }
+            currentSeries!.count++;
+            consecutivePasses = 0;
+            if (x < currentSeries!.start) {
+                currentSeries!.start = x;
+            }
+            lastFingerprint = fingerprint;
+        });
+
+        const bandHeight = 4
+        const bandOpacities = [ 0.34, 0.64 ];
+
         const svg = d3.select(container)
             .append("svg")
             .attr("width", width)
@@ -92,7 +139,7 @@ export class PhaseHistoryGraph {
 
                 const rticks: number[] = [];
                 for (let i = 0; i < nticks.length; i = i + n) {
-                rticks.push(nticks[i]);
+                    rticks.push(nticks[i]);
                 }
 
                 rticks.unshift(first);
@@ -117,6 +164,47 @@ export class PhaseHistoryGraph {
 
         svg.append("g").call(xAxis)
 
+        // add group for bands below all marks
+        svg.append("g")
+            .attr("class", "error-bands")
+            .selectAll()
+            .data(fingerprintSeries.filter(band => band.count >= minimumBandCount && (band.end - band.start >= minimumBandLength)))
+            .enter()
+            .append("g")
+            .each((band, i, nodes) => {
+                const bandNode = d3.select(nodes[i]);
+                const bandOpacity = bandOpacities[i % bandOpacities.length];
+                const xStart = Math.min(xScale(band.start), xScale(band.end));
+                const xEnd = Math.max(xScale(band.start), xScale(band.end));
+                const bandWidth = (xEnd - xStart) + 8;
+                bandNode
+                    .attr("transform", `translate(${xStart - 4.5},${yScale(band.y)! + 16 - bandHeight/2})`);
+                bandNode.append("rect")
+                    .attr("width", bandWidth)
+                    .attr("height", bandHeight)
+                    .attr("fill", band.color)
+                    .attr("opacity", bandOpacity)
+                    .attr("stroke", "none");
+                bandNode.append("line")
+                    .attr("x1", 1)
+                    .attr("x2", 1)
+                    .attr("y1", bandHeight)
+                    .attr("y2", bandHeight*3)
+                    .attr("stroke", band.color)
+                    .attr("stroke-width", 2)
+                    .attr("stroke-opacity", bandOpacity);
+                bandNode.append("text")
+                    .attr("text-anchor", "start")
+                    .style("alignment-baseline", "left")
+                    .style("font-family", "Horde Open Sans Regular")
+                    .style("font-size", 9)
+                    .style("font-weight", 800)
+                    .attr("y", 18)
+                    .attr("x", 3)
+                    .attr("fill", band.color)
+                    .text(`#${band.fingerprint.substring(0,2)}`);
+            });
+
         const g = svg.append("g")
             .selectAll()
             .data(_d3.group(I, i => Y[i]))
@@ -151,18 +239,23 @@ export class PhaseHistoryGraph {
             SessionCircle(circle.node() as any, items, items[index].id, 5, !!this.phaseSessionId && items.some(i => i.id === this.phaseSessionId));
         }
 
-        const tooltip = this.tooltip = d3.select(container)
-            .append("div")
-            .attr("id", "tooltip")
-            .style("display", "none")
-            .style("background-color", modeColors.background)
-            .style("border", "solid")
-            .style("border-width", "1px")
-            .style("border-radius", "3px")
-            .style("border-color", dashboard.darktheme ? "#413F3D" : "#2D3F5F")
-            .style("padding", "6px")
-            .style("position", "absolute")
-            .style("pointer-events", "none");
+        // Attempt to select an existing div with id 'tooltip' from the body, or create one if it doesn't exist
+        let tooltip: DivSelectionType = d3.select("body").select("div#tooltip");
+        if (tooltip.empty()) {
+            tooltip = d3.select("body")
+                .append("div")
+                .attr("id", "tooltip")
+                .style("display", "none")
+                .style("background-color", modeColors.background)
+                .style("border", "solid")
+                .style("border-width", "1px")
+                .style("border-radius", "3px")
+                .style("border-color", dashboard.darktheme ? "#413F3D" : "#2D3F5F")
+                .style("padding", "6px")
+                .style("position", "absolute")
+                .style("pointer-events", "none");
+        }
+        this.tooltip = tooltip;
 
         const closestData = (x: number, y: number): PhaseSessionResult[] | undefined => {
 
@@ -202,11 +295,15 @@ export class PhaseHistoryGraph {
             if (items.length > 1) {
                 desc += renderToString(SessionListSelector(items, index, '(mouse wheel to select)'));
             }
-            desc += `#${items.length > 1? index + 1 : ''} ${closest.outcome} <br/>`;
+            desc += `${closest.outcome}${items.length > 1? `::${index + 1}` : ''}`
+            if (!!closest.errorFingerprint) {
+                desc += ` #${closest.errorFingerprint}`;
+            }
+            desc += `<br/>`;
             desc += `on Commit ${closest.commitId} <br/>`;
             desc += `${date} <br/>`;
 
-            return desc;
+            return `<span style="color: ${modeColors.text};">${desc}</span>`;
         }
 
         const handleMouseMove = (event: any) => {
@@ -236,8 +333,13 @@ export class PhaseHistoryGraph {
                 const tx = xScale(timeStamp);
                 const ty = yScale(metaIndex)!;
                 const desc = getInfoSession(closestItems, index);
-
-                this.updateTooltip(tx, ty, desc);
+                // Compute the bounding rectangle of the container to get its absolute position
+                const containerRect = container.getBoundingClientRect();
+                const offsetX = containerRect.left + window.scrollX;
+                const offsetY = containerRect.top + window.scrollY;
+                const absTx = tx + offsetX;
+                const absTy = ty + offsetY;
+                this.updateTooltip(absTx, absTy, desc, closestItems.length > 1);
             }
 
         }
@@ -247,6 +349,7 @@ export class PhaseHistoryGraph {
             if (this.tooltipItem) {
                 svg.selectAll(".circleScale").attr("transform", `scale(1)`);
                 tooltip.style("display", "none");
+                tooltip.classed("tooltip-block-scroll", false);
                 let index = this.tooltipItem.items.length > 1? this.tooltipItem.items.findIndex(item => item.id === this.phaseSessionId) : 0;
                 if (index < 0) index = 0;
                 reDrawCircles(this.tooltipItem.items, index);
@@ -264,6 +367,7 @@ export class PhaseHistoryGraph {
                 const closestItems = closestData(mouseX, mouseY);
                 if (closestItems && this.onClickCallback) {
                     this.onClickCallback(closestItems[this.tooltipItem.index]);
+                    handleMouseLeave(undefined);
                 }
             }
 
@@ -285,7 +389,7 @@ export class PhaseHistoryGraph {
 
                     reDrawCircles(this.tooltipItem.items, this.tooltipItem.index);
                     const desc = getInfoSession(this.tooltipItem.items, this.tooltipItem.index);
-                    this.updateTooltip(undefined, undefined, desc);
+                    this.updateTooltip(undefined, undefined, desc, true);
                 }
             }
             
@@ -298,7 +402,7 @@ export class PhaseHistoryGraph {
 
     }
 
-    updateTooltip(x?: number, y?: number, html?: string) {
+    updateTooltip(x?: number, y?: number, html?: string, blockWheel?: boolean) {
         if (!this.tooltip) {
             return;
         }
@@ -311,6 +415,7 @@ export class PhaseHistoryGraph {
 
         this.tooltip
             .style("display", "block")
+            .classed("tooltip-block-scroll", blockWheel ?? false)
             .html(html ?? "")
             .style("position", `absolute`)
             .style("width", `max-content`)
@@ -358,7 +463,7 @@ const SessionCircle = (container: SVGGElement, sessions: PhaseSessionResult[], s
 
     if (selected) {
         g.append("circle")
-            .attr("stroke", dashboard.darktheme ? "#c1c4c9" : "#3d4654 ")
+            .attr("stroke", dashboard.darktheme ? "#c1c4c9af" : "#3d4654af")
             .attr("fill", "transparent")
             .attr("stroke-width", 2)
             .attr("r", radius + 3);
diff --git a/Engine/Source/Programs/Horde/HordeDashboard/plugins/test-automation/testAutomationCommon.tsx b/Engine/Source/Programs/Horde/HordeDashboard/plugins/test-automation/testAutomationCommon.tsx
--- a/Engine/Source/Programs/Horde/HordeDashboard/plugins/test-automation/testAutomationCommon.tsx
+++ b/Engine/Source/Programs/Horde/HordeDashboard/plugins/test-automation/testAutomationCommon.tsx
@@ -1,6 +1,6 @@
 // Copyright Epic Games, Inc. All Rights Reserved.
 
-import { mergeStyles, mergeStyleSets, Text, ComboBox, IComboBoxOption, SelectableOptionMenuItemType, IContextualMenuProps, IContextualMenuItem, IconButton, IContextualMenuListProps, Stack, Checkbox, FontIcon, Modal, DefaultButton, ScrollablePane, SelectionZone, DetailsList, Selection, DetailsListLayoutMode, SelectionMode, IColumn, PrimaryButton, Label, TextField, MessageBar, MessageBarType } from "@fluentui/react";
+import { mergeStyles, mergeStyleSets, Text, ComboBox, IComboBoxOption, SelectableOptionMenuItemType, IContextualMenuProps, IContextualMenuItem, IconButton, IContextualMenuListProps, Stack, Checkbox, FontIcon, Modal, DefaultButton, ScrollablePane, SelectionZone, DetailsList, Selection, DetailsListLayoutMode, SelectionMode, IColumn, PrimaryButton, Label, TextField, MessageBar, MessageBarType, ContextualMenuItemType } from "@fluentui/react";
 import { MetadataRef, PhaseSessionResult, TestDataHandler, TestNameRef, TestPhaseStatus, TestSessionDetails, TestSessionResult } from "./testData";
 import dashboard, { StatusColor } from "horde/backend/Dashboard";
 import { memo, useCallback, useEffect, useMemo, useState } from "react";
@@ -375,6 +375,36 @@ export const PhasesStatusPie = (phases: TestPhaseStatus[], radius: number): JSX.
     return StatusPie(stack, radius);
 }
 
+export const PhaseValues = (phases: TestPhaseStatus[], style?: any): JSX.Element => {
+    let phasesFailedCount = 0;
+    let phasesUnspecifiedCount = 0;
+    let phasesSkippedCount = 0
+    phases.forEach(p => {
+        switch (p.outcome) {
+            case TestPhaseOutcome.Skipped:
+                ++phasesSkippedCount;
+                break;
+            case TestPhaseOutcome.Failed:
+                ++phasesFailedCount;
+                break;
+            case TestPhaseOutcome.Interrupted:
+            case TestPhaseOutcome.Unknown:
+            case TestPhaseOutcome.NotRun:
+                ++phasesUnspecifiedCount;
+                break;
+        }
+    });
+
+    const statusColors = getStatusColors();
+    const border = `1px solid ${dashboard.darktheme ? "#ffffff3f" : "#0000003f"}`;
+    const background = dashboard.darktheme ? "#ffffff0f" : "#0000000f";
+    return <Stack horizontal style={style} tokens={{childrenGap: 3}}>
+                <Text style={{color: statusColors.get(StatusColor.Unspecified), whiteSpace: 'nowrap', border: border, borderRadius: '5px', padding: '0px 4px', background: background}} title={statusTexts.get(TestOutcome.Unspecified)}>{statusTexts.get(TestOutcome.Unspecified)?.at(0)}: {phasesUnspecifiedCount}</Text>
+                <Text style={{color: statusColors.get(StatusColor.Failure), whiteSpace: 'nowrap', border: border, borderRadius: '5px', padding: '0px 4px', background: background}} title={statusTexts.get(TestOutcome.Failure)}>{statusTexts.get(TestOutcome.Failure)?.at(0)}: {phasesFailedCount}</Text>
+                <Text style={{color: statusColors.get(StatusColor.Skipped), whiteSpace: 'nowrap', border: border, borderRadius: '5px', padding: '0px 4px', background: background}} title={statusTexts.get(TestOutcome.Skipped)}>{statusTexts.get(TestOutcome.Skipped)?.at(0)}: {phasesSkippedCount}</Text>
+            </Stack>
+}
+
 export const SessionValues = (session: TestSessionResult, style?: any): JSX.Element => {
     const statusColors = getStatusColors();
     return <Text style={style}>
@@ -426,7 +456,7 @@ export const getSessionColor = (item: TestSessionResult | PhaseSessionResult) =>
     }
 }
 
-export const PhaseHistoryGraphWidget: React.FC<{test: TestNameRef, phaseKey: string, meta: MetadataRef, streamId: string, sessionId?: string, handler: TestDataHandler, onClick: PhaseSessionCallback, withoutHeader?: boolean}> = ({ test, phaseKey, meta, streamId, sessionId, handler, onClick, withoutHeader }) => {
+export const PhaseHistoryGraphWidget: React.FC<{test: TestNameRef, phaseKey: string, meta: MetadataRef, streamId: string, sessionId?: string, handler: TestDataHandler, onClick: PhaseSessionCallback, withoutHeader?: boolean, reverse?: boolean}> = memo(({ test, phaseKey, meta, streamId, sessionId, handler, onClick, withoutHeader, reverse }) => {
     const [container, setContainer] = useState<HTMLDivElement | null>(null);
     const [state, setState] = useState<{ graph?: PhaseHistoryGraph }>({});
 
@@ -438,7 +468,7 @@ export const PhaseHistoryGraphWidget: React.FC<{test: TestNameRef, phaseKey: str
     useEffect(() => {
         // refresh on meta or sessionId update
         try {
-            container && state.graph?.render(container, meta, sessionId, withoutHeader);
+            container && state.graph?.render(container, meta, sessionId, withoutHeader, reverse);
         } catch (err) {
             console.error(err);
         }
@@ -450,8 +480,8 @@ export const PhaseHistoryGraphWidget: React.FC<{test: TestNameRef, phaseKey: str
 
     const graph_container_id = `test_history_graph_container_${test!.id}_${meta.id}_${phaseKey}_${streamId}`;
 
-    return <div id={graph_container_id} style={{ userSelect: "none", position: "relative" }} ref={(ref: HTMLDivElement) => setContainer(ref)}/>
-}
+    return <div id={graph_container_id} className="phase-history" style={{ userSelect: "none", position: "relative" }} ref={(ref: HTMLDivElement) => setContainer(ref)}/>
+});
 
 export const getStreamOptions = (streams: string[], onClick: (stream: string) => void, selected?: string): IContextualMenuProps => {
     const projectItems: Map<string, Map<string, IContextualMenuItem[]>> = new Map();
@@ -560,6 +590,10 @@ export const getPhaseOutcomeOptions = (): ICheckListOption[] => {
     ]
 }
 
+export interface IPhaseFilter {
+    isMatch(names: string[], outcomes?: Set<TestPhaseOutcome>, tags?: string[]): boolean;
+}
+
 class PhaseFiltersContext {
     constructor() {
         this._outcomeOptions = getPhaseOutcomeOptions();
@@ -571,7 +605,7 @@ class PhaseFiltersContext {
         return this._outcomeOptions
     }
 
-    filterTestPhases(testId: string, phases: TestPhaseStatus[]) {
+    filterTestPhases(testId: string, phases: IPhaseFilter[]) {
         const outcomes = this.getFilteredOutcomes();
         const names = this.getFilteredNames(testId);
         const tags = this.getFilteredTags(testId);
@@ -637,6 +671,23 @@ class PhaseFiltersContext {
 }
 export const phaseFiltersContext = new PhaseFiltersContext();
 
+class UserNavData {
+    constructor() {
+        this._userNavData = new Map();
+    }
+
+    getData(key: string) {
+        return this._userNavData.get(key);
+    }
+
+    setData(key: string, data: any) {
+        this._userNavData.set(key, data);
+    }
+
+    private _userNavData: Map<string, any>;
+}
+export const userNavData = new UserNavData();
+
 export const CheckListOption: React.FC<{options: ICheckListOption[], onChange: (items: ICheckListOption[]) => void}> = ({options, onChange}) => {
     const [items, setItems] = useState<ICheckListOption[]>([]);
     const [update, setUpdate] = useState<number>(0);
@@ -707,6 +758,69 @@ export const CheckListOption: React.FC<{options: ICheckListOption[], onChange: (
             />
 }
 
+export const CheckListOptionSmall: React.FC<{options: ICheckListOption[], onChange: (items: ICheckListOption[]) => void, onDismissed?: () => void}> = ({options, onChange, onDismissed}) => {
+    const [update, setUpdate] = useState<number>(0);
+
+    const isFiltering : boolean = useMemo(() => options.filter(i => i.itemType !== ContextualMenuItemType.Header).some(i => i.checked), [options, update]);
+
+    const fitlerOptions: IContextualMenuProps = {
+        items: options,
+        onMenuDismissed: onDismissed,
+        onRenderMenuList: (props: IContextualMenuListProps) =>
+                <Stack tokens={{childrenGap: 6}} style={{margin: '3px 2px'}}>
+                    <DefaultButton
+                        iconProps={{ iconName: "MultiSelect" }}
+                        text="Clear All"
+                        disabled={!isFiltering}
+                        onClick={() => {
+                            options.forEach(i => i.checked = false);
+                            setUpdate(update + 1);
+                            onChange(options);
+                        }}
+                        style={{ fontSize: 12, borderRadius: 6, height: 22, width: 120, paddingLeft: 5 }}
+                    />
+                    {props.items.map(item => {
+                        return (item.itemType === ContextualMenuItemType.Header ?
+                                    <Text key={item.key} style={{color: styles.highlightFilter.backgroundColor}}>{item.text}</Text>
+                                    : <Checkbox
+                                        key={item.key}
+                                        label={item.text}
+                                        checked={item.checked}
+                                        onChange={
+                                            (ev, checked) => {
+                                                item.checked = checked;
+                                                setUpdate(update + 1);
+                                                onChange(options);
+                                            }
+                                        }
+                                        onRenderLabel={(props) => {
+                                            return <Stack verticalAlign="center" tokens={{childrenGap: 2}} style={{background: item.color ? `${item.color}60` : undefined}}>
+                                                        <Text style={{padding: '0 4px'}}>{props?.label}</Text>
+                                                        <Stack style={{ height: 3, backgroundColor: item.color }}></Stack>
+                                                    </Stack>
+                                        }}
+                                    />)
+                        }
+                    )}
+                </Stack>
+    }
+
+    const iconName : string = isFiltering? "Filter" : "FilterSettings";
+    const iconColor : string = isFiltering? styles.highlightFilter.backgroundColor : "";
+
+    return  <IconButton
+                iconProps={{iconName: iconName}}
+                menuProps={fitlerOptions}
+                style={{
+                    fontSize:16, height: 28,
+                    border: '1px solid', borderColor: dashboard.darktheme ? "#4D4C4B" : "#6D6C6B", borderRadius: 4,
+                    backgroundColor: dashboard.darktheme ? "rgba(255, 255, 255, 0.1)" : "",
+                    color: dashboard.darktheme ? 'white' : 'grey'
+                }}
+                styles={{icon: {color: iconColor}}}
+            />
+}
+
 type ArtifactHref = {
     name: string;
     reference: string;
@@ -894,9 +1008,9 @@ export class DomainNameColors {
         return this._colorWheel.domain();
     }
 
-    get(name?: string) {
+    get(name?: string, opacity?: string) {
         if (!name) return;
-        return this._colorWheel(name) as string;
+        return this._colorWheel(name) as string + (opacity ? opacity : "");
     }
 
     private _colorRange = dashboard.darktheme ? [...d3.schemeDark2, ...d3.schemeCategory10] : [...d3.schemeSet2,...d3.schemeTableau10];
@@ -1046,18 +1160,61 @@ export const UserPickerModal: React.FC<{ defaultUser?: GetUserResponse, onSelect
 };
 
 export const SessionListSelector = (items: TestSessionResult[] | PhaseSessionResult[], index: number, text?: string): JSX.Element => {
-    const selectedColor = dashboard.darktheme ? "white" : "black";
+    const selectedFilter = dashboard.darktheme ? "brightness(1.5)" : "brightness(0.5)";
     const backColor = dashboard.darktheme ? "#2D2B29" : "#EDEBE9";
     const textColor = dashboard.darktheme ? "#979797ff" : "#707070ff";
 
     return <Stack>
                 <Stack style={{height: 3, backgroundColor: backColor}} grow horizontal tokens={{childrenGap: items.length < 20? 1: 0}}>
                     { items.map( (item, i) => {
-                            const itemColor = index === i? selectedColor : getSessionColor(item);
-                            return <Stack key={`s-${i}`} grow verticalFill style={{backgroundColor: itemColor}}></Stack>
+                            const itemFilter = index === i? selectedFilter : undefined;
+                            return <Stack key={`s-${i}`} grow verticalFill style={{backgroundColor: getSessionColor(item), filter: itemFilter}}></Stack>
                         } )
                     }
                 </Stack>
                 {!!text && <Text style={{color: textColor}} variant="xSmall">&ensp;{text}</Text>}
             </Stack>
 }
+
+export class CircularList<T> {
+    constructor(items: T[]) {
+        this._items = items;
+        this._cursor = -1;
+    }
+
+    setCursor(item: T | undefined): number {
+        if (item === undefined) {
+            this._cursor = -1;
+        } else if (this.selected !== item) {
+            this._cursor = this._items.findIndex(i => i === item);
+        }
+
+        return this._cursor;
+    }
+
+    get selected(): T | undefined {
+        if (this._cursor === -1) return;
+        return this._items[this._cursor];
+    }
+
+    next(): T {
+        if (++this._cursor === this._items.length) this._cursor = 0;
+        return this._items[this._cursor];
+    }
+
+    back(): T {
+        if (--this._cursor < 0) this._cursor = this._items.length - 1;
+        return this._items[this._cursor];
+    }
+
+    get cursor() {
+        return this._cursor;
+    }
+
+    get length() {
+        return this._items.length;
+    }
+
+    private _cursor: number;
+    private _items: T[];
+}
diff --git a/Engine/Source/Programs/Horde/HordeDashboard/plugins/test-automation/testAutomationEvents.tsx b/Engine/Source/Programs/Horde/HordeDashboard/plugins/test-automation/testAutomationEvents.tsx
--- a/Engine/Source/Programs/Horde/HordeDashboard/plugins/test-automation/testAutomationEvents.tsx
+++ b/Engine/Source/Programs/Horde/HordeDashboard/plugins/test-automation/testAutomationEvents.tsx
@@ -3,12 +3,11 @@
 import { Checkbox, DefaultButton, DirectionalHint, FontIcon, HoverCard, HoverCardType, Image, IPlainCardProps, Link, mergeStyles, mergeStyleSets, Spinner, SpinnerSize, Stack, Text } from "@fluentui/react";
 import { Artifacts, EmbeddedHTML, EventEntry, EventLevel, ImageCompare, LogPropertyType, URLLink } from "./testEventsModel";
 import { TestPhaseOutcome } from "./api";
-import { memo, useCallback, useEffect, useMemo, useState } from "react";
+import { memo, useCallback, useEffect, useState } from "react";
 import { ArtifactFactory, DeviceRef, TestPhaseStatus } from "./testData";
 import dashboard, { StatusColor } from "horde/backend/Dashboard";
 import { getShortNiceTime, msecToElapsed } from "horde/base/utilities/timeUtils";
 import { DashboardPreference } from "horde/backend/Api";
-import { isBright } from "horde/base/utilities/colors";
 import { PhaseArtifactsModal, styles, getStatusColors } from "./testAutomationCommon";
 
 let eventStyles: any = undefined;
@@ -19,6 +18,7 @@ export const getEventStyling = () => {
     const errorColorPref = dashboard.getPreference(DashboardPreference.ColorError);
     const warningColorPref = dashboard.getPreference(DashboardPreference.ColorWarning);
     const highlight = dashboard.darktheme ? "brightness(1.2)" : "brightness(0.9)";
+    const strongHighlight = dashboard.darktheme ? "brightness(1.8)" : "brightness(0.8)";
 
     const eventClass = mergeStyles({
         fontSize: "12px",
@@ -66,6 +66,7 @@ export const getEventStyling = () => {
             }
         },
         failureButton: {
+            minWidth: 16,
             backgroundColor: errorColorPref ? errorColorPref : dashboard.darktheme ? "#9D1410" : "#EC4C47",
             borderStyle: "hidden",
             color: "#FFFFFF",
@@ -103,6 +104,9 @@ export const getEventStyling = () => {
                 }
             }
         },
+        highlighted: {
+            filter: strongHighlight
+        }
     });
 
     return eventStyles;
diff --git a/Engine/Source/Programs/Horde/HordeDashboard/plugins/test-automation/testAutomationPhaseCompare.tsx b/Engine/Source/Programs/Horde/HordeDashboard/plugins/test-automation/testAutomationPhaseCompare.tsx
--- a/Engine/Source/Programs/Horde/HordeDashboard/plugins/test-automation/testAutomationPhaseCompare.tsx
+++ b/Engine/Source/Programs/Horde/HordeDashboard/plugins/test-automation/testAutomationPhaseCompare.tsx
@@ -30,6 +30,7 @@ const onClickItem = (session: PhaseSessionResult) => {
 
 export const PhaseComparisonPane: React.FC<{ test: TestNameRef, meta?: MetadataRef, phaseKey: string, streamId?: string, focusedSessionId?: string, handler: TestDataHandler, onDismiss: () => void }> = observer(({test, meta, phaseKey, streamId, focusedSessionId, handler, onDismiss}) => {
     const hordeTheme = getHordeTheme();
+    const isFromGridView = !!handler.getSearchParam('grid');
 
     comparisonContext.subscribe();
     handler.subscribeToPhaseQueryLoading();
@@ -134,7 +135,7 @@ export const PhaseComparisonPane: React.FC<{ test: TestNameRef, meta?: MetadataR
                             if (!!test) {
                                 const phase = test.phases.find(p => p.key === phaseKey);
                                 if (!!phase) {
-                                    handler.queryPhase(phase, comparisonContext.streamId!, true)
+                                    handler.queryPhaseSessions([phase], comparisonContext.streamId!, true)
                                         .then(() => {
                                             // force auto selection once the stream is loaded
                                             comparisonContext.setLoadingPhase(false);
@@ -168,7 +169,7 @@ export const PhaseComparisonPane: React.FC<{ test: TestNameRef, meta?: MetadataR
                             comparisonContext.commitId = test.commitId;
                             const phase = test.phases.find(p => p.key === phaseKey);
                             if (!!phase) {
-                                handler.queryPhase(phase, comparisonContext.streamId!, true)
+                                handler.queryPhaseSessions([phase], comparisonContext.streamId!, true)
                                     .then(() => comparisonContext.setPhaseStatus(phase));
                                 return;
                             }
@@ -292,7 +293,7 @@ export const PhaseComparisonPane: React.FC<{ test: TestNameRef, meta?: MetadataR
                         {!!comparisonContext.metaKey && !!selectedMeta &&
                             <Stack style={{ paddingLeft: 3 }} horizontal verticalAlign="end">
                                 <Stack style={{ marginLeft: 30, marginBottom: 18, width: 10, cursor: "default" }} title="History">
-                                    <FontIcon style={{ fontSize: 15 }} iconName={"History"}/>
+                                    <FontIcon style={{ fontSize: 15 }} iconName="Timeline"/>
                                 </Stack>
                                 <PhaseHistoryGraphWidget
                                     test={test}
@@ -302,6 +303,7 @@ export const PhaseComparisonPane: React.FC<{ test: TestNameRef, meta?: MetadataR
                                     meta={selectedMeta}
                                     onClick={onClickItem}
                                     handler={handler}
+                                    reverse={isFromGridView}
                                 />
                                 <Stack grow horizontalAlign="end" style={{ marginBottom: 16, marginRight: 16}}>
                                     <FontIcon
diff --git a/Engine/Source/Programs/Horde/HordeDashboard/plugins/test-automation/testAutomationPhaseGridView.tsx b/Engine/Source/Programs/Horde/HordeDashboard/plugins/test-automation/testAutomationPhaseGridView.tsx
new file mode 100644
--- /dev/null
+++ b/Engine/Source/Programs/Horde/HordeDashboard/plugins/test-automation/testAutomationPhaseGridView.tsx
@@ -0,0 +1,954 @@
+// Copyright Epic Games, Inc. All Rights Reserved.
+
+import { IconButton, Stack, Text, Spinner, SpinnerSize, Link, FontIcon, DefaultButton, TooltipHost, Icon, ITag, TagPicker, IPickerItemProps, List, IPage, IDetailsHeaderProps, IRenderFunction, Sticky, StickyPositionType, IColumn, DetailsList, SelectionMode, DetailsListLayoutMode, ScrollablePane, ScrollToMode } from "@fluentui/react";
+import React, { useState, useEffect, useCallback, useRef, memo, useMemo } from "react";
+import { observer } from "mobx-react-lite";
+import { TestDataHandler, TestSessionDetails, MetadataRef, PhaseSessionResult } from "./testData";
+import dashboard from "horde/backend/Dashboard";
+import { CheckListOption, CheckListOptionSmall, getPhaseSessionStatusColor, ICheckListOption, phaseFiltersContext, PhaseHistoryGraphWidget, PhasesStatusPie, PhaseValues, StreamSelector, styles, tagColors, testViewIcons, TestViewType, userNavData } from "./testAutomationCommon";
+import { getEventStyling } from "./testAutomationEvents";
+import { MetaSelectorContext, PhaseStatusGroup, MetadataItem } from "./phaseGrid";
+
+type PhaseNavInfo = { key: string, testId: string };
+
+const metaWidth = 20;
+const phaseRowHeight = 44;
+
+const PhaseHeader: React.FC<{ group: PhaseStatusGroup, onClick: () => void }> = memo(({ group, onClick }) => {
+
+    const tags = useMemo(() => {
+        return group.phase.tagRefs?.values().toArray().map(t => t.name);
+    }, [group]);    
+
+    return <Stack disableShrink style={{cursor: "pointer", width: '100%', height: phaseRowHeight, paddingLeft: 16, paddingTop: 4}} verticalAlign="start" onClick={onClick}>
+                <Stack horizontal verticalAlign="center" horizontalAlign="start" title={group.phase.name} style={{width: '100%', paddingRight: 3}}>
+                    <Text style={{fontSize: 13, fontWeight: 600, whiteSpace: 'nowrap', overflow: 'hidden', textOverflow: 'ellipsis', direction: 'rtl', textAlign: 'right'}}>{group.phase.name}</Text>
+                </Stack>
+                <Stack horizontal verticalAlign="center" horizontalAlign="start" style={{overflow: 'hidden', textOverflow: 'ellipsis'}}>
+                    {!!tags &&
+                        <Stack horizontal verticalAlign="center" tokens={{childrenGap: 6}} style={{paddingLeft: 6, paddingRight: 3}}>
+                            {tags.map(tag => <Stack key={`tag-${tag}`} style={{backgroundColor: tagColors.getTagColor(tag), borderRadius: 5, padding: '0px 4px', opacity: 0.8}}><Text style={{fontSize: 11, whiteSpace: 'nowrap'}}>#{tag}</Text></Stack>)}
+                        </Stack>
+                    }
+                </Stack>
+            </Stack>
+});
+
+const PhaseStatusWidget: React.FC<{ group: PhaseStatusGroup, context: MetaSelectorContext, onClick: (phaseKey: string, meta?: MetadataItem) => void }> = memo(observer(({ group, context, onClick }) => {
+    const [errorsCaught, setErrorsCaught] = useState<Map<string, number> | undefined>(undefined);
+
+    useEffect(() => {
+        if (context.focusPhase === group.phase.key) {
+            context.focusPhase = undefined;
+        }
+        return context.registerOnPhaseDataLoaded(group.phase.key, () => {
+            if (group.errorFingerprints) {
+                const today = Date.now();
+                const errorsCaught = context.getFirstTimeErrorsCaught(group);
+                setErrorsCaught(new Map(errorsCaught.entries().toArray().map(([key, value]) => [key, Math.round((today - value.getTime()) / (1000 * 60 * 60 * 24))])));
+            }
+        });
+    }, [group, context]);
+
+    context.subscribeFilteredPhases();
+    
+    return <div ref={(ref) => { group.componentRef = ref; }} style={{width: '100%', height: '100%' }}>
+                <Stack horizontal verticalAlign="center" style={{cursor: "pointer", height: '100%' }}>                
+                { context.metas.filter(m => m.visible).map(m => {
+                    const status = group.getMetadataStatus(m);
+                    const iconColor = status ? getPhaseSessionStatusColor(status) : undefined;
+                    const errorDays = errorsCaught?.get(status?.errorFingerprint ?? "") ?? 0;
+                    const title = `${status ? status.outcome : 'N/A'}${!!m.name ? ` - ${m.name}` : ''}`;
+                    return <Stack key={`meta-${m.meta.id}`}
+                                className={`${getMetaCellClassName(m.meta.id)} ${getPhaseCellClassName(group.phase.key)}`}
+                                style={{ height: '100%', paddingLeft: 4, paddingRight: 4, width: metaWidth, position: 'relative' }}
+                                horizontalAlign="center" verticalAlign="center"
+                                title={title}
+                                onClick={() => onClick(group.phase.key, m)}
+                                >                                    
+                                    { !!status && <FontIcon style={{ fontSize: 12 * (errorDays/context.daysQueried + 1), color: iconColor }} iconName="Square" /> }
+                                    { !!status && !!status.errorFingerprint && !!errorsCaught && errorDays > 0 &&
+                                        <Text style={{ fontSize: 8, color: 'white', position: 'absolute', cursor: 'pointer', userSelect: "none" }}>{errorDays >= context.daysQueried ? `${context.daysQueried}+` : errorDays}</Text>
+                                    }
+                                    { !!status && !!status.errorFingerprint &&
+                                        <Text style={{ fontSize: 9, color: iconColor, position: 'absolute', bottom: 2, right: 1, cursor: 'pointer', userSelect: "none", fontWeight: 800 }}>#{status.errorFingerprint.substring(0,2)}</Text>
+                                    }
+                                    { !status && <div style={{ minHeight: 12, minWidth: 12 }}/> }
+                            </Stack>
+                    })
+                }
+                </Stack>
+            </div>
+}));
+
+const PhaseHistoryWidget: React.FC<{ group: PhaseStatusGroup, context: MetaSelectorContext, handler: TestDataHandler}> = memo(observer(({ group, context, handler }) => {
+    const [loading, setLoading] = useState<boolean>(false);
+
+    context.subscribeSelectedMeta();
+
+    useEffect(() => {
+        setLoading(true);
+        handler.queryPhaseSessions(group.metadata.values().toArray(), handler.stream!, true).finally(() => {
+            setLoading(false);
+            context.notifyOnPhaseDataLoaded(group.phase.key);
+        });
+    }, [group, handler]);
+
+    const onClickSession = useCallback((session: PhaseSessionResult) => {
+        const testStatus = MetaSelectorContext.getTestStatus(handler, context.test);
+        if (!context.selectedMeta || !testStatus) return;
+        const metaStatus = testStatus.sessions.get(context.selectedMeta.meta);
+        const testDataId = metaStatus?.history.find(s => s.id === session.sessionId)?.testDataId;
+        if (!!testDataId) {
+            handler.setSearchParam('session', testDataId);
+            handler.setSearchParam('phase', group.phase.key);
+            context.setSelectedPhase(group.phase.key);
+        }
+    }, [group, handler, context]);
+
+    return <Stack grow horizontalAlign="start" verticalAlign="center" style={{height: phaseRowHeight}} onClick={(ev) => ev.stopPropagation()}>
+                {!loading && !!context.test && !!context.selectedMeta && context.selectedMeta.visible && !!group.getMetadataStatus(context.selectedMeta) &&
+                    <PhaseHistoryGraphWidget
+                        test={context.test}
+                        streamId={handler.stream!}
+                        phaseKey={group.phase.key}
+                        meta={context.selectedMeta.meta}
+                        sessionId={group.metadata.get(context.selectedMeta)?.session?.id}
+                        handler={handler}
+                        withoutHeader={true}
+                        reverse={true}
+                        onClick={onClickSession}
+                    />
+                }
+                {!!loading && !!context.selectedMeta && <Spinner size={SpinnerSize.small} style={{marginLeft: 32}} />}
+            </Stack>
+}));
+
+const onRenderSuggestionsItem = (item: ITag) => {
+    return <Stack style={{height: 24, padding: 4}}>
+                <Text title={item.name} style={{textOverflow: 'ellipsis', overflow: 'hidden', whiteSpace: 'nowrap', maxWidth: 300, fontSize: 12}}>{item.name}</Text>
+            </Stack>
+}
+
+const onRenderItem = (props: IPickerItemProps<ITag>) => {
+    const item = props.item;
+    const isTag = item.name.startsWith("#");
+    return <Stack key={`picker_item_${item.name}`} horizontal verticalAlign="center" verticalFill={true} style={{ marginRight: 2 }}>
+                <TooltipHost content={item.name} delay={0} closeDelay={500}>
+                    <FontIcon
+                        style={{ fontSize: 13, padding: 4, paddingTop: 5, cursor: 'pointer', backgroundColor: isTag? tagColors.getTagColor(getTextFromTag(item)) : styles.defaultFilter.backgroundColor, borderRadius: 3, color: 'white'}}
+                        iconName="Filter"
+                        title="Remove"
+                        onClick={props.onRemoveItem} />
+                </TooltipHost>
+            </Stack>
+}
+
+const getTextFromTag = (item: ITag) => item.name;
+
+const PhaseGroupFilter: React.FC<{groups: PhaseStatusGroup[], selected?: string[], onChange: (items: string[]) => void, onFocus?: () => void, onBlur?: () => void}> = ({ groups, selected, onChange, onFocus, onBlur }) => {
+    const [filters, setFilters] = useState<ITag[]>([]);
+    const phaseTokens = useRef<ITag[]>();
+    const phaseTags = useRef<ITag[]>();
+
+    useEffect(() => {
+        const tokens: ITag[] = (new Set(groups.map(g => (g.phase.name.split('.'))).flat())).keys().map(token => ({key: token.toLowerCase(), name: token})).toArray().sort((a, b) => a.key.length - b.key.length);
+        const names: ITag[] = groups.map(g => ({key: `@${g.phase.name.toLowerCase()}`, name: `@${g.phase.name}`}));
+        const tags: ITag[] = (new Set(groups.map(g => g.metadata.values().toArray().flatMap(status => status.tags ?? [])).flat())).keys().map(t => ({key: `#${t.toLowerCase()}`, name: `#${t}`})).toArray().sort((a, b) => a.key.length - b.key.length);
+        phaseTokens.current = [...tokens, ...names];
+        phaseTags.current = tags;
+        selected && setFilters(selected.map(i => ({key: i, name: i})));
+    }, [groups]);
+    
+    const filterSelectedTests = useCallback((filterText: string): ITag[] => {
+        if (filterText.length < 2) return [];
+        const lowerText = filterText.toLowerCase();
+        if (filterText.startsWith("#")) {
+            const closestTags = phaseTags.current?.values().filter((tag) => (tag.key as string).indexOf(lowerText) >= 0).take(10).toArray() ?? [];
+            return closestTags;
+        }
+        const closestTokens = phaseTokens.current?.values().filter((tag) => (tag.key as string).indexOf(lowerText) >= 0).take(10).toArray() ?? [];
+        const isTokenExist = closestTokens.find((tag) => tag.key === lowerText);
+        if (!isTokenExist) closestTokens.splice(0, 0, {key: lowerText, name: filterText});
+        return closestTokens;
+    }, [groups]);
+
+    return <TagPicker
+                styles={{
+                    text: { minWidth: 210, maxWidth: 450, maxHeight: 28, display: 'inline-flex', flexWrap: 'nowrap' },
+                    input: { height: 28, minWidth: 40, maxWidth: 200 },
+                    itemsWrapper: { marginLeft: 3, marginTop: 2, marginRight: 3, marginBottom: 2, height: 28, maxWidth: 400, overflow: 'hidden' }
+                }}
+                onRenderItem={onRenderItem}
+                onRenderSuggestionsItem={onRenderSuggestionsItem}
+                removeButtonAriaLabel="Remove"
+                selectionAriaLabel="Filter phases"
+                selectedItems={filters}
+                onResolveSuggestions={filterSelectedTests}
+                getTextFromItem={getTextFromTag}
+                onChange={(tokens) => {
+                    let items = tokens ?? [];
+                    const uniqueTokens = new Set(items.map(i => i.key));
+                    items = uniqueTokens.keys().map(k => items.find(i => i.key === k)!).toArray();
+                    setFilters(items);
+                    onChange(items.map(t => t.name));
+                }}
+                inputProps={{ placeholder: "Filter names and #tags", onFocus: onFocus, onBlur: onBlur }}
+            />
+}
+
+const MetaSelectorTitle: React.FC<{ stateContext: MetaSelectorContext, onFilterChange?: (items?: string[]) => void }> = observer(({ stateContext, onFilterChange }) => {
+    const [hoveredMeta, setHoveredMeta] = useState<MetadataItem | undefined>(undefined);
+
+    stateContext.subscribeMetaFilter();
+    stateContext.subscribeFilteredPhases();
+    stateContext.subscribeSelectedMeta();
+
+    useEffect(() => {
+        return stateContext.registerOnMetaHovered(setHoveredMeta);
+    }, [stateContext]);
+
+    const options = stateContext.filterOptions;
+
+    useEffect(() => {
+        const gridFilterItems = stateContext.getSearchGridFilterItems();
+        let changed = false;
+        options.forEach(o => {
+            const newChecked = gridFilterItems?.includes(o.text!) ?? false;
+            if (newChecked !== o.checked) {
+                changed = true;
+                o.checked = newChecked;
+            }
+        });
+        if (changed) {
+            onMetaFilterChange(options, true);
+        }
+    }, [options, stateContext.searchUpdated]);
+
+    const onMetaFilterChange = useCallback((items: ICheckListOption[], forceUpdate?: boolean) => {
+        stateContext.updateMetasVisibility(items);
+        stateContext.updateMetaFilter();
+        stateContext.findForkingMetaKey();
+        const gridFilterItems = items.filter(i => i.checked).map(i => i.text!);
+        stateContext.setSearchGridFilterItems(gridFilterItems, forceUpdate);
+        onFilterChange?.(gridFilterItems);
+    }, [stateContext, onFilterChange]);
+
+    const onMetaValueClick = useCallback((value: string) => {
+        const option = options.find(o => o.text === value);
+        if (option) {
+            option.checked = !option.checked;
+            onMetaFilterChange(options);
+        }
+    }, [options, onMetaFilterChange]);
+
+    const onMetaValueBeforeMainClick = useCallback((value: string) => {
+        let changed = false;
+        stateContext.getMetaValuesBeforeForkingKey().reverse().some(t => {
+            const option = options.find(o => o.text === t);
+            if (option) {
+                option.checked = false;
+                changed = true;
+            }
+            return t === value;
+        });
+        if (changed) {
+            onMetaFilterChange(options);
+        }
+    }, [options, onMetaFilterChange, stateContext]);
+
+    const latestCommit = useMemo(() => stateContext.filteredPhases?.at(0)?.metadata?.values().reduce((latest, s) => (s.session?.commitOrder ?? 0) > latest.order? {order: s.session!.commitOrder, id: s.session!.commitId} : latest, {order: 0, id: "NA"}), [stateContext.filteredPhases]);
+
+    const phaseStatus = useMemo(() => new Map(stateContext.metas.map(m => [m.meta, stateContext.filteredPhases.map(p => p.getMetadataStatus(m)).filter(s => !!s)])), [stateContext.metas, stateContext.filteredPhases]);
+
+    const displayedMeta = hoveredMeta ?? (stateContext.selectedMeta?.visible ? stateContext.selectedMeta : undefined) ?? (stateContext.metas.filter(m => m.visible).length === 1 ? stateContext.metas.find(m => m.visible) : undefined);
+    const displayedStatus = displayedMeta ? phaseStatus.get(displayedMeta.meta) : undefined;
+
+    return <Stack verticalAlign="center" grow horizontal tokens={{childrenGap: 10}} style={{height: '100%', anchorName: '--meta-values-container'}}>
+                <Stack horizontal tokens={{childrenGap: 5}} style={{height: 21, position: 'fixed', borderLeft: `1px solid ${dashboard.darktheme ? '#4D4C4B' : '#edebe9'}`, left: 'calc(anchor(--meta-values left) - 0.48px)', bottom: 'anchor(--meta-values-container bottom)'}}>
+                    <Stack horizontal>
+                        {!!stateContext.metaKeysBeforeFork && stateContext.metaKeysBeforeFork.length > 0 &&
+                            <Stack horizontal style={{ position: 'fixed', right: `calc(anchor(--meta-values left) + 2px)`}}>
+                                {stateContext.getMetaValuesBeforeForkingKey().map(v => (
+                                    <Stack horizontal key={`meta-value-${v}`}>
+                                        <Stack style={{background: stateContext.metaValueColors?.get(v, "60"), cursor: "pointer"}}
+                                            onClick={() => onMetaValueBeforeMainClick(v)}
+                                        >
+                                            <Text style={{whiteSpace: 'nowrap', padding: '0 10px 0 4px', position: "relative", height: 18}}>{v}
+                                                {!!options.find(o => o.text === v)?.checked &&
+                                                    <FontIcon iconName="Filter" style={{fontSize: 8, position: "absolute", right: 2, top: 2, color: styles.highlightFilter.backgroundColor}} />
+                                                }
+                                            </Text>
+                                            <Stack style={{background: stateContext.metaValueColors?.get(v), height: 3}}></Stack>
+                                        </Stack>
+                                        <Stack style={{background: dashboard.darktheme ? "#ffffff3f" : "#0000003f", paddingBottom: 2}} verticalAlign="center">
+                                            <FontIcon iconName="ChevronRight" style={{fontSize: 12, cursor: "default"}} />
+                                        </Stack>
+                                    </Stack>
+                                ))}
+                            </Stack>
+                        }
+                        {!!displayedMeta?.visible &&
+                            <Stack horizontal>
+                                {displayedMeta.meta.getValuesExcept(stateContext.getCommonKeysToFork()).map((v, index, self) =>
+                                    <Stack horizontal key={`meta-value-${v}`}
+                                        className={stateContext.getMetasContainingValue(v).map(m => getMetaCellClassName(m.meta.id)).join(' ')}
+                                    >
+                                        <Stack style={{background: stateContext.metaValueColors?.get(v, "60"), cursor: "pointer", position: "relative"}}
+                                            onClick={() => onMetaValueClick(v)}
+                                        >
+                                            <Text style={{whiteSpace: 'nowrap', padding: '0 10px 0 4px', position: "relative", height: 18}}>{v}
+                                                {!!options.find(o => o.text === v)?.checked &&
+                                                    <FontIcon iconName="Filter" style={{fontSize: 8, position: "absolute", right: 2, top: 2, color: styles.highlightFilter.backgroundColor}} />
+                                                }
+                                            </Text>
+                                            <Stack style={{background: stateContext.metaValueColors?.get(v), height: 3}}></Stack>
+                                        </Stack>
+                                        {index < self.length - 1 &&
+                                            <Stack style={{background: dashboard.darktheme ? "#ffffff3f" : "#0000003f", paddingBottom: 2}} verticalAlign="center">
+                                                <FontIcon iconName="ChevronRight" style={{fontSize: 12, cursor: "default"}} />
+                                            </Stack>
+                                        }
+                                    </Stack>
+                                )}
+                            </Stack>
+                        }
+                        {!displayedMeta && !!stateContext.forkingMetaKey &&
+                            <Stack horizontal>
+                                {stateContext.getForkingMetaKeyValues().map(v => {
+                                        const visibleMetasCount = stateContext.getVisibleMetasContainingForkingKeyValue(v).length;
+                                        const isTooTight = v ? v.length / visibleMetasCount > 2 : false; // heuristic to avoid too tight spacing
+                                        return (
+                                            <Stack key={`meta-value-${v}`} title={v}
+                                                style={{background: v ? stateContext.metaValueColors?.get(v, "60") : undefined, cursor: "pointer", minWidth: metaWidth*visibleMetasCount}}
+                                                onClick={() => onMetaValueClick(v)}
+                                            >
+                                                <Stack className={stateContext.getMetasContainingValue(v).map(m => getMetaCellClassName(m.meta.id)).join(' ')} style={{height: 18, position: 'relative'}} horizontal>
+                                                    <Text style={{
+                                                        whiteSpace: 'nowrap',
+                                                        overflow: 'hidden',
+                                                        textOverflow: 'ellipsis',
+                                                        width: isTooTight ? undefined : '100%',
+                                                        transform: isTooTight ? 'rotate(-35deg)' : 'none',
+                                                        transformOrigin: '0% 100%',
+                                                        position: 'absolute',
+                                                        left: isTooTight ? 8 : 4,
+                                                    }}>{v ?? " "}</Text>
+                                                </Stack>
+                                                <Stack style={{background: v ? stateContext.metaValueColors?.get(v) : undefined, height: 3}}></Stack>
+                                            </Stack>
+                                        );
+                                    })
+                                }
+                            </Stack>
+                        }
+                    </Stack>
+                    {(!!displayedMeta?.visible && !!displayedStatus) &&
+                        <Stack horizontal verticalAlign="start" tokens={{childrenGap: 5}}>
+                            {PhaseValues(displayedStatus, {cursor: "default"})}
+                            <Text style={{whiteSpace: 'nowrap', cursor: "default"}}>- CL: {displayedStatus.at(0)?.commitId ?? "N/A"}</Text>
+                            {!!latestCommit && latestCommit.id === displayedStatus.at(0)?.commitId && <FontIcon style={{ fontSize: 11, cursor: "default" }} title="Latest commit" iconName="Star" />}
+                        </Stack>
+                    }
+                </Stack>
+                <Stack grow horizontalAlign="end" verticalAlign="center">
+                    <TooltipHost content='Filter metadata'>
+                        <CheckListOptionSmall key={`options-${stateContext.metaFilterUpdated}`} options={options} onChange={onMetaFilterChange}/>
+                    </TooltipHost>
+                </Stack>
+            </Stack>
+});
+
+const MetaSelector: React.FC<{ onClickMeta?: (meta?: MetadataItem) => void, stateContext: MetaSelectorContext }> = observer(({ onClickMeta, stateContext }) => {
+
+    stateContext.subscribeMetaFilter();
+    stateContext.subscribeFilteredPhases();
+    stateContext.subscribeSelectedMeta();
+    
+    const eventStyles = getEventStyling();
+    const forkingMetaKey = stateContext.forkingMetaKey;
+    const phaseStatus = useMemo(() => new Map(stateContext.metas.map(m => [m.meta, stateContext.filteredPhases.map(p => p.getMetadataStatus(m)).filter(s => !!s)])), [stateContext.metas, stateContext.filteredPhases]);
+
+    const metaBackgroundColor = dashboard.darktheme ? "#1F2223" : "#f3f2f1";
+
+    return <Stack verticalAlign="start" horizontal style={{paddingTop: 1, anchorName: '--meta-values'}} >
+                {!!stateContext.metas &&
+                    stateContext.metas.filter(m => m.visible).map(m => {
+                        return <Stack key={`meta-${m.meta.id}`} verticalAlign="start"
+                                    style={{ backgroundColor: stateContext.selectedMeta?.meta === m.meta ? metaBackgroundColor : undefined, cursor: onClickMeta ? 'pointer' : 'default' }}
+                                    className={`${stateContext.selectedMeta?.meta === m.meta ? eventStyles.highlighted : undefined} ${getMetaCellClassName(m.meta.id)}`}
+                                    onMouseOver={() => stateContext.onMetaHovered?.(m)}
+                                    onMouseLeave={() => stateContext.onMetaHovered?.(undefined)}
+                                    onClick={() => stateContext.selectedMeta?.meta === m.meta ? onClickMeta?.(undefined) : onClickMeta?.(m)}
+                               >
+                                    <Stack                                        
+                                        verticalAlign="start"
+                                        horizontalAlign="center"
+                                        style={{ width: metaWidth, height: 24, padding: 2, background: !!forkingMetaKey ? stateContext.metaValueColors?.get(m.meta.getValue(forkingMetaKey), "60") : undefined }}
+                                    >{PhasesStatusPie(phaseStatus.get(m.meta)!, 9)}</Stack>
+                                </Stack>
+                    })
+                }
+            </Stack>
+});
+
+const headerStyles = {
+    root: {
+        padding: 0,
+        height: 26,
+        background: 'transparent !important',
+        // Remove hover background color on header cells
+        selectors: {
+            '.ms-DetailsHeader-cell:hover': {
+                background: 'transparent',
+            },
+            '.ms-DetailsHeader-cell': {
+                borderRight: `1px solid ${dashboard.darktheme ? '#4D4C4B' : '#edebe9'}`,
+                padding: 0,
+                height: 'fit-content',
+            },
+            '.ms-DetailsHeader-cell:last-child': {
+                borderRight: 'none'
+            }
+        },
+    }
+};
+
+const listStyles = {
+    root: {
+        selectors: {
+            // Row background color
+            '.ms-DetailsRow': {
+                background: 'transparent !important'
+            },
+            '.ms-DetailsRow:hover': {
+                background: dashboard.darktheme ? "#ffffff1f !important" : "#0000001f !important"
+            },
+            // Vertical borders between row cells
+            '.ms-DetailsRow-cell': {
+                borderRight: `1px solid ${dashboard.darktheme ? '#4D4C4B' : '#edebe9'}`,
+                padding: 0
+            },
+            '.ms-DetailsRow-cell:last-child': {
+                borderRight: 'none'
+            }
+        }
+    }
+}
+
+const getPhaseCellClassName = (phaseKey: string) => `phase-cell-${phaseKey}`;
+const getMetaCellClassName = (id: string) => `meta-cell-${id}`;
+
+const getListCellHoverStyleSelectors = (ids: string[]) => {
+    const selectors: any = {};
+    ids.forEach(id => {
+        const cellClassName = getMetaCellClassName(id);
+        const selectorRule = `&:has(.${cellClassName}:hover) .${cellClassName}`;
+        selectors[selectorRule] = {
+            background: dashboard.darktheme ? "#ffffff1f !important" : "#0000001f !important"
+        }
+    });
+    return selectors;
+}
+
+const getOnRenderDetailsHeader = (onRenderPhaseName: (column: PhaseColumn) => JSX.Element | null, onRenderMetaStatus: (column: PhaseColumn) => JSX.Element | null, onRenderHistory: (column: PhaseColumn) => JSX.Element | null) => {
+    return (props: IDetailsHeaderProps, defaultRender: IRenderFunction<IDetailsHeaderProps>) => {
+        if (!defaultRender) {
+            return null;
+        }
+        return (
+            <Sticky stickyPosition={StickyPositionType.Header}>
+                {defaultRender({...props!,
+                        styles: headerStyles,
+                        onRenderColumnHeaderTooltip: (tooltipProps) => {
+                            const column = tooltipProps?.column as PhaseColumn;
+                            if (column.key === 'phase-name') {
+                                return onRenderPhaseName(column);
+                            }
+                            if (column.key === 'phase-meta-status') {
+                                return onRenderMetaStatus(column);
+                            }
+                            if (column.key === 'phase-selected-meta-history') {
+                                return onRenderHistory(column);
+                            }
+                            return null;
+                        }
+                    })
+                }
+            </Sticky>
+        );
+    }
+}
+
+type PhaseColumn = IColumn & {
+    focusedPhase?: string;
+    handler?: TestDataHandler;
+    context?: MetaSelectorContext;
+}
+
+const getPhaseColumn = (context: MetaSelectorContext, handler: TestDataHandler, onPhaseNameRender: (item: PhaseStatusGroup) => React.ReactNode, onMetaStatusRender: (item: PhaseStatusGroup) => React.ReactNode, onHistoryRender: (item: PhaseStatusGroup) => React.ReactNode): PhaseColumn[] => {
+    const statusColumnWidth = metaWidth * (context.metas.filter(m => m.visible).length) - 20;
+    return [
+        {
+            key: 'phase-name',
+            name: 'Phase',
+            minWidth: 200,
+            isResizable: true,
+            isCollapsible: true,
+            isSorted: false,
+            isSortedDescending: true,
+            onRender: onPhaseNameRender,
+            handler: handler,
+            context: context,
+        },
+        {
+            key: 'phase-meta-status',
+            name: 'Status',
+            minWidth: statusColumnWidth === 0 ? 0.5 : statusColumnWidth,
+            maxWidth: statusColumnWidth === 0 ? 0.5 : statusColumnWidth,
+            isResizable: false,
+            isCollapsible: false,
+            isSorted: false,
+            onRender: onMetaStatusRender,
+            context: context,
+        },
+        {
+            key: 'phase-selected-meta-history',
+            name: 'History',
+            minWidth: 1000,
+            isResizable: false,
+            isCollapsible: false,
+            isSorted: false,
+            onRender: onHistoryRender,
+            context: context,
+        }
+    ];
+}
+
+export const TestPhasesGridView: React.FC<{ testId: string, handler: TestDataHandler, onDismiss: () => void }> = observer(({ testId, handler, onDismiss }) => {
+    const listBoundaryRef = useRef<HTMLDivElement | null>(null);
+    const listRef = useRef<List>();
+    const contextRef = useRef<MetaSelectorContext | undefined>(undefined);
+
+    const test = useMemo(() => handler.testStreams.get(handler.stream!)?.tests.find(t => t.id === testId), [testId, handler.stream]);
+    const testStatus = MetaSelectorContext.getTestStatus(handler, test);
+    const metaSelectorContext = useMemo(() => new MetaSelectorContext(handler, test), [handler, test, testStatus]);
+    contextRef.current = metaSelectorContext;
+
+    const onPhaseNameRender = useCallback((item: PhaseStatusGroup) => {
+        return <PhaseHeader group={item} onClick={() => { onPhaseSelected(item.phase.key); }} />
+    }, [metaSelectorContext]);
+    const onMetaStatusRender = useCallback((item: PhaseStatusGroup) => {
+        return <PhaseStatusWidget group={item} context={metaSelectorContext} onClick={onPhaseSelected} />
+    }, [metaSelectorContext]);
+    const onHistoryRender = useCallback((item: PhaseStatusGroup) => {
+        return !!testStatus? <PhaseHistoryWidget group={item} context={metaSelectorContext} handler={handler} /> : null;
+    }, [metaSelectorContext, handler, testStatus]);
+
+    const columns = useMemo(() => getPhaseColumn(metaSelectorContext, handler, onPhaseNameRender, onMetaStatusRender, onHistoryRender)
+    , [metaSelectorContext, handler, onPhaseNameRender, onMetaStatusRender, onHistoryRender, metaSelectorContext.filteredPhasesUpdated]);
+
+    const onPhaseNameHeaderRender = useCallback((column: PhaseColumn) => {
+        const handler = column.handler!;
+        const context = column.context!;
+        const streamId = handler.stream!;
+        return <Stack horizontal tokens={{childrenGap: 8}} verticalAlign="start" style={{height: '100%', padding: '0px 12px'}}>
+                    <StreamSelector
+                        streams={handler.availableStreams.filter((s) => s !== streamId)}
+                        selected={streamId}
+                        onClick={(stream) => {
+                            handler.selectStream(stream);
+                            context.setPhaseGroups([]);
+                        }}
+                    />
+                    <Link style={{padding: '0px 5px'}} title="open test health view" onClick={() => {
+                            handler.setSearchParam('view', TestViewType.Health);
+                            handler.setSearchParam('health', handler.getSearchParam('grid') as string);
+                            handler.removeSearchParam('grid');
+                            handler.removeSearchParam('phase');
+                            handler.removeSearchParam('metagrid');
+                        }}>
+                        <FontIcon iconName={testViewIcons.get(TestViewType.Health)} style={{fontSize: 15, paddingRight: 4}} />
+                    </Link>
+                </Stack>
+    }, []);
+    const onMetaStatusHeaderRender = useCallback((column: PhaseColumn) => {
+        const context = column.context!;
+        return <MetaSelector
+                    stateContext={context}
+                    onClickMeta={(item) => {
+                        onPhaseSelected(undefined, item);
+                    }}
+                />
+    }, [metaSelectorContext]);
+    const onHistoryHeaderRender = useCallback((column: PhaseColumn) => {
+        return <Stack horizontalAlign="start" verticalAlign="end">
+                    { !!test && !!testStatus && <PhaseHistoryGraphWidget
+                            test={test}
+                            streamId={handler.stream!}
+                            phaseKey={'fake'}
+                            meta={new MetadataRef('fake', {}, 0)}
+                            handler={handler}
+                            reverse={true}
+                            onClick={() => {}}
+                        />
+                    }
+                </Stack>
+    }, [test, handler, testStatus]);
+    const onRenderDetailsHeader = useMemo(() => getOnRenderDetailsHeader(onPhaseNameHeaderRender
+    , onMetaStatusHeaderRender, onHistoryHeaderRender), [onPhaseNameHeaderRender, onMetaStatusHeaderRender, onHistoryHeaderRender, handler.searchUpdated]);
+
+    handler.subscribeToTestDataQueryLoading();
+    handler.subscribeToPhaseQueryLoading();
+    handler.subscribeToQueryLoading();
+    handler.subscribeToSearch();
+    metaSelectorContext.subscribePhaseGroups();
+    metaSelectorContext.subscribeFilteredPhases();
+    metaSelectorContext.subscribeSelectedPhase();
+    metaSelectorContext.subscribeSelectedMeta();
+
+    // collect test phase groups
+    useEffect(() => {
+        if (!!testStatus && !handler.phaseQueryLoading)
+        {
+            handler.queryTestPhases(testStatus).then(() => {
+                const testDataRequests = testStatus.getMetadataLastSessions().map(([_, session]) => {
+                    return handler.queryTestData(session.testDataId!);
+                }).filter(request => !!request);
+                if (testDataRequests && testDataRequests.length > 0) {
+                    handler.setTestDataQueryLoading(true);
+                    const tagRefMap = new Map(test!.tagRefs?.values().toArray().map(t => [t.name, t]));
+                    Promise.all(testDataRequests).then((responses) => {
+                        const sessionDetails = (responses.filter(response => !!response) as TestSessionDetails[]);
+                        // reorg data by unique phase and keep order
+                        const seen = new Map<string, PhaseStatusGroup>();
+                        const phaseGroups: PhaseStatusGroup[] = [];
+                        sessionDetails.forEach(session => {
+                            if (!session.meta) return;
+                            session.phases.forEach(phase => {
+                                const phaseRef = test!.phases?.get(phase.key);
+                                if (!phaseRef) return;
+                                phase.phaseRef = phaseRef;
+                                // update tags
+                                const tags = phase.tags?.map(t => tagRefMap.get(t)).filter(t => !!t);
+                                if (tags) phaseRef.updateTags(tags);
+                                // populate phase group
+                                let group = seen.get(phaseRef.key);
+                                if (!group) {
+                                    group = new PhaseStatusGroup(phaseRef);
+                                    seen.set(phaseRef.key, group);
+                                    phaseGroups.push(group);
+                                };
+                                group.addMetadata(metaSelectorContext.metas.find(m => m.meta === session.meta)!, phase);
+                            });
+                        });
+                        metaSelectorContext.setPhaseGroups(phaseGroups);
+                        metaSelectorContext.setFilteredPhases(phaseFiltersContext.filterTestPhases(testId, phaseGroups) as PhaseStatusGroup[]);
+                        handler.setTestDataQueryLoading(false);
+                    });
+                }
+                });
+        }
+    }, [testStatus, metaSelectorContext, handler])
+
+    const onPagesUpdated = useCallback((pages: IPage<PhaseStatusGroup>[]) => {
+        // Extract all loaded items from all pages
+        const loaded = pages.filter(page => page.items && !page.isSpacer).flatMap(page => page.items ?? []);
+        //console.log('Loaded items:', loaded.length);
+        if (loaded.length === 0) return;
+        // scroll to focus phase if it is not visible
+        if (!metaSelectorContext.focusPhase) return;
+        const phaseIndex = metaSelectorContext.filteredPhases.findIndex((g) => g.phase.key === metaSelectorContext.focusPhase);
+        if (phaseIndex === -1) return;
+        if (!isPhaseVisible(phaseIndex)) {
+            // force the list to load the page that contains the target item
+            console.log('onPagesUpdated', metaSelectorContext.focusPhase);
+            listRef.current?.scrollToIndex(phaseIndex, undefined, ScrollToMode.center);
+        } else {
+            // validate the scroll to index was achieved
+            metaSelectorContext.focusPhase = undefined;
+        }
+    }, [metaSelectorContext]);
+
+    const isPhaseVisible = useCallback((phaseIndex: number) => {
+        // Check if the phase row is already visible by checking if the componentRef exists and is within visible bounds
+        const phaseGroup = metaSelectorContext.filteredPhases[phaseIndex];
+        if (phaseGroup?.componentRef && listBoundaryRef.current) {
+            const rowRect = (phaseGroup.componentRef as HTMLElement).getBoundingClientRect();
+            const listRect = listBoundaryRef.current.getBoundingClientRect();
+            return rowRect.top > (listRect.top + headerStyles.root.height) && rowRect.bottom < listRect?.bottom;
+        }
+        return false;
+    }, [metaSelectorContext]);
+
+    const onPhaseSelected = useCallback((phaseKey?: string, metaItem?: MetadataItem) => {
+        if (!!metaItem) {
+            handler.setSearchParam('metagrid', metaItem.meta.id);
+            metaSelectorContext.setSelectedMeta(metaItem);
+        } else if (!phaseKey) {
+            handler.removeSearchParam('metagrid');
+            metaSelectorContext.setSelectedMeta(undefined);
+        }
+
+        if (!!phaseKey && metaSelectorContext.setSelectedPhase(phaseKey)) {
+            if (listRef.current) {
+                const phaseIndex = metaSelectorContext.filteredPhases.findIndex((g) => g.phase.key === phaseKey);
+                if (phaseIndex === -1 || isPhaseVisible(phaseIndex))  return;
+                listRef.current.scrollToIndex(phaseIndex, undefined, ScrollToMode.center);
+            }
+        }
+    }, [handler, metaSelectorContext]);
+
+    useEffect(() => {
+        const metaParam = handler.getSearchParam('metagrid') as string | undefined;
+        metaSelectorContext.setSelectedMetaById(metaParam);
+
+        // check if there is a phase nav info and if it is, force select the phase to focus
+        const phaseNavInfo = userNavData.getData('phaseNavInfo') as PhaseNavInfo | undefined;
+        if (!!phaseNavInfo) {
+            if (phaseNavInfo.testId === metaSelectorContext.test?.id) {
+                userNavData.setData('phaseNavInfo', undefined);
+                onPhaseSelected(phaseNavInfo.key);
+            }
+        }
+    }, [metaSelectorContext, handler.searchUpdated]);
+
+    useEffect(() => {
+        // on unmount
+        return () => {
+            if (contextRef.current?.selectedPhase) {
+                // save the selected phase to focus when the component is mounted again
+                userNavData.setData('phaseNavInfo', { key: contextRef.current.selectedPhase, testId: contextRef.current.test?.id ?? '' });
+            }
+        }
+    }, []);
+
+    const eventStyles = getEventStyling();
+    const buttonNoneTextColor = dashboard.darktheme ? "#616e85" : "#949898";
+    const failureButtonTextColor = "#F9F9FB";
+    const interruptedButtonTextColor = failureButtonTextColor;
+    const warningButtonTextColor = failureButtonTextColor;
+
+    const failureEnabled = !!metaSelectorContext.failureCursor?.length;
+    const failureText = failureEnabled ? metaSelectorContext.failureCursor!.length.toString() : "No Failure";
+
+    const interruptedEnabled = !!metaSelectorContext.interruptedCursor?.length;
+    const warningEnabled = !!metaSelectorContext.warningCursor?.length;
+
+    const localListStyles = useMemo(() => {
+        const styles = structuredClone(listStyles);
+        styles.root.selectors = {
+                ...styles.root.selectors,
+                ...getListCellHoverStyleSelectors(metaSelectorContext.metas.map(m => m.meta.id))
+            };
+        if (!!metaSelectorContext.selectedPhase) {
+            styles.root.selectors[`.ms-DetailsRow:has(.${getPhaseCellClassName(metaSelectorContext.selectedPhase)})`] = {
+                background: dashboard.darktheme ? "#ffffff1f !important" : "#0000001f !important"
+            };
+        }
+        if (!!metaSelectorContext.selectedMeta) {
+            styles.root.selectors[`.${getMetaCellClassName(metaSelectorContext.selectedMeta.meta.id)}`] = {
+                background: dashboard.darktheme ? "#ffffff1f !important" : "#0000001f !important"
+            };
+        }
+        return styles;
+    }, [metaSelectorContext, metaSelectorContext.selectedPhaseUpdated, metaSelectorContext.selectedMetaUpdated]);
+
+    return <Stack style={{width: '100%', height: '100%'}} styles={localListStyles}>
+                {!testStatus && (handler.queryLoading || handler.phaseQueryLoading || handler.testDataQueryLoading) &&
+                        <Stack horizontalAlign='center' style={{ padding: 12 }} tokens={{ childrenGap: 10 }}>
+                            <Text style={styles.textLarge}>Loading Data</Text>
+                            <Spinner size={SpinnerSize.large} />
+                        </Stack>
+                }
+                {!handler.queryLoading && !handler.phaseQueryLoading && !handler.testDataQueryLoading &&
+                    <Stack grow style={{height: '100%'}}>
+                        <Stack horizontal verticalAlign="start" tokens={{childrenGap: 10}}>
+                            <Stack style={{ paddingLeft: 8, minWidth: 550 }} tokens={{childrenGap: 5}}>
+                                <Text style={{ fontSize: 16, fontWeight: 600, whiteSpace: 'nowrap', overflow: 'hidden', textOverflow: 'ellipsis', cursor: "default" }}>{test?.name ?? testId}{!!metaSelectorContext.commonMetaString && ` - ${metaSelectorContext.commonMetaString}`}</Text>
+                            </Stack>
+
+                            <Stack horizontal verticalAlign="center" style={{ height: '100%', paddingLeft: 6 }} grow>
+                                <MetaSelectorTitle
+                                    stateContext={metaSelectorContext}
+                                    onFilterChange={(items) => {
+                                        metaSelectorContext.setFilteredPhases(phaseFiltersContext.filterTestPhases(testId, metaSelectorContext.phaseGroups) as PhaseStatusGroup[]);
+                                    }}
+                                />
+                            </Stack>
+
+                            <Stack horizontal verticalAlign='center' style={{ height: '100%' }}>
+                                <Stack horizontal verticalAlign='center' style={{position: 'relative'}}>
+                                    <TooltipHost content={failureEnabled ? 'Navigate to next failure': ''}>
+                                        <DefaultButton
+                                            disabled={!failureEnabled} className={failureEnabled ? eventStyles.failureButton : eventStyles.failureButtonDisabled}
+                                            style={{ color: failureEnabled ? failureButtonTextColor : buttonNoneTextColor, height: 28, whiteSpace: 'nowrap', fontSize: 12, paddingLeft: 8, paddingRight: 8, borderRadius: failureEnabled ? '4px 0 0 4px' : '4px' }}
+                                            text={failureText}
+                                            onClick={() => {
+                                                const [nextPhase, nextMetadata] = metaSelectorContext.getNextFailure();
+                                                if (nextPhase) {
+                                                    onPhaseSelected(nextPhase, nextMetadata);
+                                                }
+                                            }} >
+                                            {failureEnabled &&
+                                                <Icon style={{ fontSize: 19, paddingLeft: 4, paddingRight: 0 }} iconName='ChevronDown' />
+                                            }
+                                        </DefaultButton>
+                                    </TooltipHost>
+                                    {failureEnabled &&
+                                        <div style={{backgroundColor: failureButtonTextColor, height: 18, width: 1, position: 'absolute', right: -0.5, zIndex: 1 }} />
+                                    }
+                                </Stack>
+                                {failureEnabled &&
+                                    <TooltipHost content={'Navigate to previous failure'}>
+                                        <IconButton
+                                            disabled={!failureEnabled} className={failureEnabled ? eventStyles.failureButton : eventStyles.failureButtonDisabled}
+                                            style={{ color: failureButtonTextColor, height: 28, fontSize: 12, padding: 2, borderRadius: '0 4px 4px 0' }}
+                                            iconProps={{ iconName: 'ChevronUp' }}
+                                            onClick={() => {
+                                                const [previousPhase, previousMetadata] = metaSelectorContext.getPreviousFailure();
+                                                if (previousPhase) {
+                                                    onPhaseSelected(previousPhase, previousMetadata);
+                                                }
+                                            }} />
+                                    </TooltipHost>
+                                }
+                            </Stack>
+
+                            {interruptedEnabled &&
+                                <Stack horizontal verticalAlign='center' style={{ height: '100%' }}>
+                                    <Stack horizontal verticalAlign='center' style={{position: 'relative'}}>
+                                        <TooltipHost content='Navigate to next interrupted'>
+                                            <DefaultButton
+                                                className={eventStyles.interruptedButton}
+                                                style={{ color: interruptedButtonTextColor, height: 28, whiteSpace: 'nowrap', fontSize: 12, paddingLeft: 8, paddingRight: 8, borderRadius: '4px 0 0 4px' }}
+                                                text={metaSelectorContext.interruptedCursor?.length.toString()}
+                                                onClick={() => {
+                                                    const [nextPhase, nextMetadata] = metaSelectorContext.getNextInterrupted();
+                                                    if (nextPhase) {
+                                                        onPhaseSelected(nextPhase, nextMetadata);
+                                                    }
+                                                }} >
+                                                <Icon style={{ fontSize: 19, paddingLeft: 4, paddingRight: 0 }} iconName='ChevronDown' />
+                                            </DefaultButton>
+                                        </TooltipHost>
+                                        <div style={{backgroundColor: interruptedButtonTextColor, height: 18, width: 1, position: 'absolute', right: -0.5, zIndex: 1 }} />
+                                    </Stack>
+                                    <TooltipHost content={'Navigate to previous interrupted'}>
+                                        <IconButton
+                                            className={eventStyles.interruptedButton}
+                                            style={{ color: interruptedButtonTextColor, height: 28, fontSize: 12, padding: 2, borderRadius: '0 4px 4px 0' }}
+                                            iconProps={{ iconName: 'ChevronUp' }}
+                                            onClick={() => {
+                                                const [previousPhase, previousMetadata] = metaSelectorContext.getPreviousInterrupted();
+                                                if (previousPhase) {
+                                                    onPhaseSelected(previousPhase, previousMetadata);
+                                                }
+                                            }} />
+                                    </TooltipHost>
+                                </Stack>
+                            }
+
+                            {warningEnabled &&
+                                <Stack horizontal verticalAlign='center' style={{ height: '100%' }}>
+                                    <Stack horizontal verticalAlign='center' style={{position: 'relative'}}>
+                                        <TooltipHost content='Navigate to next warning'>
+                                            <DefaultButton
+                                                className={eventStyles.warningButton}
+                                                style={{ color: warningButtonTextColor, height: 28, whiteSpace: 'nowrap', fontSize: 12, paddingLeft: 8, paddingRight: 8, borderRadius: '4px 0 0 4px' }}
+                                                text={metaSelectorContext.warningCursor?.length.toString()}
+                                                onClick={() => {
+                                                    const [nextPhase, nextMetadata] = metaSelectorContext.getNextWarning();
+                                                    if (nextPhase) {
+                                                        onPhaseSelected(nextPhase, nextMetadata);
+                                                    }
+                                                }} >
+                                                <Icon style={{ fontSize: 19, paddingLeft: 4, paddingRight: 0 }} iconName='ChevronDown' />
+                                            </DefaultButton>
+                                        </TooltipHost>
+                                        <div style={{backgroundColor: warningButtonTextColor, height: 18, width: 1, position: 'absolute', right: -0.5, zIndex: 1 }} />
+                                    </Stack>
+                                    <TooltipHost content={'Navigate to previous warning'}>
+                                        <IconButton
+                                            className={eventStyles.warningButton}
+                                            style={{ color: warningButtonTextColor, height: 28, fontSize: 12, padding: 2, borderRadius: '0 4px 4px 0' }}
+                                            iconProps={{ iconName: 'ChevronUp' }}
+                                            onClick={() => {
+                                                const [previousPhase, previousMetadata] = metaSelectorContext.getPreviousWarning();
+                                                if (previousPhase) {
+                                                    onPhaseSelected(previousPhase, previousMetadata);
+                                                }
+                                            }} />
+                                    </TooltipHost>
+                                </Stack>
+                            }
+
+                            <Stack verticalAlign='center' style={{ height: '100%' }}>
+                                <TooltipHost content='Filter by outcome'>
+                                    <CheckListOption
+                                        options={phaseFiltersContext.outcomeOptions}
+                                        onChange={() => {
+                                            metaSelectorContext.setFilteredPhases(phaseFiltersContext.filterTestPhases(testId, metaSelectorContext.phaseGroups) as PhaseStatusGroup[]);
+                                        }}
+                                    />
+                                </TooltipHost>
+                            </Stack>
+
+                            <Stack verticalAlign='center' style={{ height: '100%' }}>
+                                <TooltipHost content='Filter phases by name and tag'>
+                                    <PhaseGroupFilter
+                                        groups={metaSelectorContext.phaseGroups}
+                                        selected={[...phaseFiltersContext.getFilteredNames(testId), ...phaseFiltersContext.getFilteredTags(testId)?.map(t => `#${t}`) ?? []]}
+                                        onChange={(items => {
+                                            phaseFiltersContext.setFilteredNames(testId, items);
+                                            metaSelectorContext.setFilteredPhases(phaseFiltersContext.filterTestPhases(testId, metaSelectorContext.phaseGroups) as PhaseStatusGroup[]);
+                                        })}/>
+                                </TooltipHost>
+                            </Stack>
+
+                            <Stack style={{ paddingBottom: 4 }}>
+                                <IconButton
+                                    iconProps={{ iconName: 'Cancel' }}
+                                    ariaLabel="Close phase grid view"
+                                    onClick={onDismiss}
+                                />
+                            </Stack>
+                        </Stack>
+                        <Stack style={{ paddingLeft: 16, paddingBottom: 3, overflow: 'hidden', height: '100%', position: 'relative'}} grow>
+                            {(handler.queryLoading || handler.phaseQueryLoading || handler.testDataQueryLoading) &&
+                                <Stack horizontalAlign='center' style={{ padding: 12 }} tokens={{ childrenGap: 10 }}>
+                                    <Text style={styles.textLarge}>Loading Data</Text>
+                                    <Spinner size={SpinnerSize.large} />
+                                </Stack>
+                            }
+                            {!!test && !handler.queryLoading && !handler.phaseQueryLoading && !handler.testDataQueryLoading &&
+                                <div ref={(ref) => { listBoundaryRef.current = ref; }} style={{ height: '100%' }}>
+                                    <ScrollablePane styles={{ contentContainer: {scrollbarGutter: 'stable', selectors: {'body:has(.tooltip-block-scroll) &': { overflowY: 'hidden !important' }} }}}>
+                                        <DetailsList
+                                            styles={{ root: {overflowX: 'hidden !important'}}}
+                                            items={metaSelectorContext.filteredPhases}
+                                            columns={columns}
+                                            getKey={(group: PhaseStatusGroup) => `phase-${group.phase.key}`}
+                                            listProps={{
+                                                componentRef: (list: List) => { listRef.current = list; },
+                                                onPagesUpdated: onPagesUpdated
+                                            }}
+                                            compact={true}
+                                            isHeaderVisible={true}
+                                            onRenderDetailsHeader={onRenderDetailsHeader}
+                                            selectionMode={SelectionMode.none}
+                                            layoutMode={DetailsListLayoutMode.justified}                                    
+                                        />
+                                    </ScrollablePane>
+                                </div>
+                            }
+                            {!test && !handler.queryLoading && !handler.phaseQueryLoading && !handler.testDataQueryLoading &&
+                                <Stack horizontalAlign='center' style={{ padding: 12 }} tokens={{ childrenGap: 10 }}>
+                                    <Text style={styles.textLarge}>Test not found in stream</Text>
+                                </Stack>
+                            }
+                        </Stack>
+                    </Stack>
+                }     
+            </Stack>
+});
\ No newline at end of file
diff --git a/Engine/Source/Programs/Horde/HordeDashboard/plugins/test-automation/testAutomationPhaseView.tsx b/Engine/Source/Programs/Horde/HordeDashboard/plugins/test-automation/testAutomationPhaseView.tsx
--- a/Engine/Source/Programs/Horde/HordeDashboard/plugins/test-automation/testAutomationPhaseView.tsx
+++ b/Engine/Source/Programs/Horde/HordeDashboard/plugins/test-automation/testAutomationPhaseView.tsx
@@ -10,7 +10,7 @@ import { ChangeButton } from "horde/components/ChangeButton";
 import { getShortNiceTime, msecToElapsed } from "horde/base/utilities/timeUtils";
 import { getHordeTheme } from "horde/styles/theme";
 import { GetArtifactResponse } from "horde/backend/Api";
-import { CheckListOption, getPhaseSessionStatusColor, phaseFiltersContext, PhaseHistoryGraphWidget, PhasesStatusPie, styles, tagColors, testViewIcons, TestViewType } from "./testAutomationCommon";
+import { CheckListOption, CircularList, getPhaseSessionStatusColor, phaseFiltersContext, PhaseHistoryGraphWidget, PhasesStatusPie, styles, tagColors, testViewIcons, TestViewType } from "./testAutomationCommon";
 import { projectStore } from "horde/backend/ProjectStore";
 import { JobArtifactsModal } from "horde/components/artifacts/ArtifactsModal";
 import { getJobParams, JobContext } from "./phaseCompare";
@@ -21,10 +21,11 @@ import { EventEntry } from "./testEventsModel";
 const PhaseHistoryPane: React.FC<{ test: TestNameRef, meta?: MetadataRef, phaseKey: string, streamId?: string, sessionId?: string, handler: TestDataHandler, onClickCompare: () => void }> = memo(({test, meta, phaseKey, streamId, sessionId, handler, onClickCompare}) => {
 
     const metaStatus = !!meta && !!streamId? handler.getStatusStream(streamId)?.tests.get(test)?.sessions.get(meta) : undefined;
+    const isFromGridView = !!handler.getSearchParam('grid');
    
     return <Stack style={{ paddingLeft: 3 }} horizontal verticalAlign="end" disableShrink>
                 <Stack style={{ marginLeft: 30, marginBottom: 18, width: 10, cursor: "default" }} title="History">
-                    <FontIcon style={{ fontSize: 15 }} iconName={"History"} />
+                    <FontIcon style={{ fontSize: 15 }} iconName="Timeline" />
                 </Stack>
                 {(!metaStatus || !sessionId) && <Text style={{ marginLeft: 10,marginBottom: 21}}>No history available for this phase</Text>}
                 {!!metaStatus && !!sessionId && <PhaseHistoryGraphWidget
@@ -34,6 +35,7 @@ const PhaseHistoryPane: React.FC<{ test: TestNameRef, meta?: MetadataRef, phaseK
                     phaseKey={phaseKey}
                     meta={meta!}
                     handler={handler}
+                    reverse={isFromGridView}
                     onClick={(session) => {
                         const testSession = metaStatus?.history.find(s => s.id === session.sessionId);
                         !!testSession?.testDataId && handler.setSearchParam('session', testSession.testDataId);
@@ -57,7 +59,7 @@ const TestPhaseExpandPane: React.FC<{ phase: TestPhaseStatus, onEventsLoaded?: (
     useEffect(() => {
         // query phase data
         setSessionId(undefined);
-        handler.queryPhase(phase, phase.testSession?.streamId ?? handler.stream!).then(() => {
+        handler.queryPhaseSessions([phase], phase.testSession?.streamId ?? handler.stream!).then(() => {
             setComparison(comparisonContext.isEnable);
             setSessionId(phase.session?.id ?? "");
             if (!phase.session) {
@@ -277,49 +279,6 @@ const PhaseFilter: React.FC<{phases: TestPhaseStatus[], selected?: string[], onC
             />
 }
 
-class CircularList<T> {
-    constructor(items: T[]) {
-        this._items = items;
-        this._cursor = -1;
-    }
-
-    setCursor(item: T | undefined): number {
-        if (item === undefined) {
-            this._cursor = -1;
-        } else {
-            this._cursor = this._items.findIndex(i => i === item);
-        }
-
-        return this._cursor;
-    }
-
-    get selected(): T | undefined {
-        if (this._cursor === -1) return;
-        return this._items[this._cursor];
-    }
-
-    next(): T {
-        if (++this._cursor === this._items.length) this._cursor = 0;
-        return this._items[this._cursor];
-    }
-
-    back(): T {
-        if (--this._cursor < 0) this._cursor = this._items.length - 1;
-        return this._items[this._cursor];
-    }
-
-    get cursor() {
-        return this._cursor;
-    }
-
-    get length() {
-        return this._items.length;
-    }
-
-    private _cursor: number;
-    private _items: T[];
-}
-
 export const TestPhasesView: React.FC<{ testDataId: string, handler: TestDataHandler, onDismiss: () => void }> = observer(({ testDataId, handler, onDismiss }) => {
     const [sessionDetails, setSessionDetails] = useState<TestSessionDetails | undefined>(undefined);
     const [filteredPhases, setFilteredPhases] = useState<TestPhaseStatus[]>([]);
@@ -380,7 +339,7 @@ export const TestPhasesView: React.FC<{ testDataId: string, handler: TestDataHan
                 }
                 setSessionDetails(fetchedSessionDetails);
                 const testId = fetchedSessionDetails.test?.id ?? testDataId;
-                updateFilter(phaseFiltersContext.filterTestPhases(testId, fetchedSessionDetails.phases));
+                updateFilter(phaseFiltersContext.filterTestPhases(testId, fetchedSessionDetails.phases) as TestPhaseStatus[]);
                 jobContext.current = getJobParams(fetchedSessionDetails, handler);
             }
         });
@@ -402,8 +361,8 @@ export const TestPhasesView: React.FC<{ testDataId: string, handler: TestDataHan
     }, [focusPhase]);
 
     selectedPhase.current = useMemo(() => {
-        return filteredPhases.find((phase) => phase.key === phaseKey);
-    }, [handler.searchUpdated, filteredPhases]);
+        return sessionDetails?.phases.find((phase) => phase.key === phaseKey);
+    }, [handler.searchUpdated, sessionDetails]);
 
     useEffect(() => {
         if (phaseKey !== failureCursor.current?.selected) {
@@ -588,7 +547,7 @@ export const TestPhasesView: React.FC<{ testDataId: string, handler: TestDataHan
                                     <CheckListOption
                                         options={phaseFiltersContext.outcomeOptions}
                                         onChange={() => {
-                                            updateFilter(phaseFiltersContext.filterTestPhases(testId, sessionDetails.phases));
+                                            updateFilter(phaseFiltersContext.filterTestPhases(testId, sessionDetails.phases) as TestPhaseStatus[]);
                                         }}
                                     />
                                 </TooltipHost>
@@ -601,7 +560,7 @@ export const TestPhasesView: React.FC<{ testDataId: string, handler: TestDataHan
                                         selected={[...phaseFiltersContext.getFilteredNames(testId), ...phaseFiltersContext.getFilteredTags(testId)?.map(t => `#${t}`) ?? []]}
                                         onChange={(items => {
                                             phaseFiltersContext.setFilteredNames(testId, items);
-                                            updateFilter(phaseFiltersContext.filterTestPhases(testId, sessionDetails.phases));
+                                            updateFilter(phaseFiltersContext.filterTestPhases(testId, sessionDetails.phases) as TestPhaseStatus[]);
                                         })}/>
                                 </TooltipHost>
                             </Stack>
diff --git a/Engine/Source/Programs/Horde/HordeDashboard/plugins/test-automation/testAutomationSummary.tsx b/Engine/Source/Programs/Horde/HordeDashboard/plugins/test-automation/testAutomationSummary.tsx
--- a/Engine/Source/Programs/Horde/HordeDashboard/plugins/test-automation/testAutomationSummary.tsx
+++ b/Engine/Source/Programs/Horde/HordeDashboard/plugins/test-automation/testAutomationSummary.tsx
@@ -309,7 +309,7 @@ const TestSummaryCard: React.FC<{ test: TestNameRef, sessions: [MetadataRef, Tes
         }
     ]
  
-    return <Stack className={hordeClasses.raised} style={{ cursor: "pointer", height: "fit-content", backgroundColor: dashboard.darktheme ? "#242729" : modeColors.background, padding: 12, width: 332 }} onClick={() => expandCard(!expanded)}>
+    return <Stack className={hordeClasses.raised} style={{ cursor: "pointer", height: "fit-content", backgroundColor: dashboard.darktheme ? "#242729" : modeColors.background, padding: 12, width: 335 }} onClick={() => expandCard(!expanded)}>
             <Stack horizontal verticalAlign="center" >
                 <Stack className="horde-no-darktheme" style={{ paddingTop: 1, paddingRight: 4 }}>
                     <FontIcon style={{ fontSize: 13 }} iconName={expanded ? "ChevronDown" : "ChevronRight"} />
@@ -318,13 +318,20 @@ const TestSummaryCard: React.FC<{ test: TestNameRef, sessions: [MetadataRef, Tes
                     <Text style={{maxWidth: 240, fontWeight: 600, textOverflow: 'ellipsis', overflow: 'hidden', whiteSpace: 'nowrap' }}  title={`${test.name} - ${commonMetaString}`}>{test.name}{maskMetaString && <span style={{fontSize: 11}}> - {maskMetaString}</span>}</Text>
                 </Stack>
         
-                <Stack grow />
+                <Stack grow />        
         
                 <Stack horizontalAlign="end">
-                    <DefaultButton style={{ minWidth: 52, fontSize: 10, padding: 0, height: 18 }} onClick={(ev) => {
+                    <DefaultButton style={{ minWidth: 25, fontSize: 10, padding: 1, height: 20 }} title="Grid Phase view" onClick={(ev) => {
+                        ev.stopPropagation();
+                        handler.setSearchParam('grid', test.id);
+                    }}><FontIcon style={{ fontSize: 15 }} iconName="TimelineMatrixView"/></DefaultButton>
+                </Stack>
+
+                <Stack horizontalAlign="end">
+                    <DefaultButton style={{ minWidth: 25, fontSize: 10, padding: 1, height: 20 }} title="History view" onClick={(ev) => {
                         ev.stopPropagation();
                         onHistoryClick(test, onDetailsClick);
-                    }}>History</DefaultButton>
+                    }}><FontIcon style={{ fontSize: 15 }} iconName="Timeline"/></DefaultButton>
                 </Stack>
             </Stack>
             <Stack style={{ paddingTop: 4 }}>
diff --git a/Engine/Source/Programs/Horde/HordeDashboard/plugins/test-automation/testAutomationView.tsx b/Engine/Source/Programs/Horde/HordeDashboard/plugins/test-automation/testAutomationView.tsx
--- a/Engine/Source/Programs/Horde/HordeDashboard/plugins/test-automation/testAutomationView.tsx
+++ b/Engine/Source/Programs/Horde/HordeDashboard/plugins/test-automation/testAutomationView.tsx
@@ -17,6 +17,7 @@ import { JobStepTestDataItem } from "./api";
 import { TestDataV2 } from "./testDataV2Fetcher";
 import dashboard, { StatusColor } from "horde/backend/Dashboard";
 import { TestAutomationHealth } from "./testAutomationHealth";
+import { TestPhasesGridView } from "./testAutomationPhaseGridView";
 
 const handler = new TestDataHandler();
 
@@ -456,6 +457,7 @@ export const TestAutomationView: React.FC = () => {
     const [initialized, setInitialized] = useState(false);
 
     const testDataId = searchParams.get('session');
+    const testId = searchParams.get('grid');
     const jobId = searchParams.get('job');
     const stepId = searchParams.get('step');
 
@@ -486,13 +488,21 @@ export const TestAutomationView: React.FC = () => {
                         <JobStepTestsSelectionModal jobId={jobId} stepId={stepId} />
                     }
                     {initialized && !!testDataId &&
-                        <TestPhasesView testDataId={testDataId!} handler={handler}
+                        <TestPhasesView testDataId={testDataId} handler={handler}
                             onDismiss={() => {
                                 handler.removeSearchParam('session');
                                 handler.removeSearchParam('phase');
                             }} />
                     }
-                    {initialized && !testDataId && !!handler.availableStreams.length &&
+                    {initialized && !!handler.stream && !!testId && !testDataId &&
+                        <TestPhasesGridView testId={testId} handler={handler}
+                            onDismiss={() => {
+                                handler.removeSearchParam('grid');
+                                handler.removeSearchParam('metagrid');
+                                handler.removeSearchParam('gridfilter');
+                            }} />
+                    }
+                    {initialized && !testDataId && !testId  && !!handler.availableStreams.length &&
                         <Stack horizontal grow>
                             <TestAutomationSidebarLeft/>
                             <Stack style={{ width: 'calc(100vw - 350px)', height: 'calc(100vh - 162px)', overflow: 'auto' }} grow>
diff --git a/Engine/Source/Programs/Horde/HordeDashboard/plugins/test-automation/testData.ts b/Engine/Source/Programs/Horde/HordeDashboard/plugins/test-automation/testData.ts
--- a/Engine/Source/Programs/Horde/HordeDashboard/plugins/test-automation/testData.ts
+++ b/Engine/Source/Programs/Horde/HordeDashboard/plugins/test-automation/testData.ts
@@ -7,9 +7,11 @@ import { projectStore } from 'horde/backend/ProjectStore';
 import backend from 'horde/backend';
 import { GetArtifactResponse } from 'horde/backend/Api';
 import { EventEntry } from './testEventsModel';
+import { IPhaseFilter } from './testAutomationCommon';
 
-const defaultQueryWeeks = 2;
+export const defaultQueryWeeks = 2;
 const filterRefreshMS = 200;
+export const consecutivePassesThreshold = 4;
 
 type StreamId = string;
 type TestId = string;
@@ -37,7 +39,10 @@ export interface TestVersionFetcher {
     fetchTestData(testDataId: TestDataId): Promise<TestSessionDetails | undefined>;
 
     /// fetch test phase sessions and populate test status phase session
-    fetchPhaseSessions(streamRef: TestStreamRef, status: TestMetaStatus, phaseKey: string, weeks: number): Promise<boolean>;
+    fetchPhaseSessions(streamRef: TestStreamRef, status: TestMetaStatus, phaseKeys: string[], weeks: number): Promise<boolean>;
+
+    /// fetch test phases and populate test status phase
+    fetchTestPhases(status: TestMetaStatus, phaseKeys?: string[]): Promise<void>;
 
     /// fetch test data information from job and step id
     fetchJobStepTestData(jobId: string, stepId: string): Promise<JobStepTestDataItem[] | undefined>;
@@ -173,7 +178,13 @@ export class TestSessionStatus {
 
     sortSessions() {
         /// sort session history by commit order descending
-        this.history.sort((a, b) => b.commitOrder - a.commitOrder);
+        this.history.sort((a, b) => {
+            if (b.commitOrder !== a.commitOrder) {
+                return b.commitOrder - a.commitOrder;
+            }
+            // if commitOrder is equal, sort by date descending
+            return b.start.getTime() - a.start.getTime();
+        });
     }
 
     addPhaseSession(phaseKey: PhaseKey, session: PhaseSessionResult) {
@@ -191,8 +202,14 @@ export class TestSessionStatus {
     }
 
     sortPhaseSessions(phaseId: string) {
-        /// sort phase session history by commit order descending
-        this.phaseHistories.get(phaseId)?.sort((a, b) => b.commitOrder - a.commitOrder);
+        /// sort phase session history by commit order descending, then by date descending
+        this.phaseHistories.get(phaseId)?.sort((a, b) => {
+            if (b.commitOrder !== a.commitOrder) {
+                return b.commitOrder - a.commitOrder;
+            }
+            // if commitOrder is equal, sort by date descending
+            return b.start.getTime() - a.start.getTime();
+        });
     }
 
     includeTags(tags: TestTagRef[]): boolean {
@@ -298,7 +315,7 @@ export class DeviceRef {
     metadata: Metadata;
 }
 
-export abstract class TestPhaseStatus {
+export abstract class TestPhaseStatus implements IPhaseFilter {
     constructor (key: string, name: string, start: Date, duration: number, devices: DeviceRef[], outcome: TestPhaseOutcome, sessionDetails: TestSessionDetails) {
         this.key = key;
         this._name = name;
@@ -368,6 +385,7 @@ export abstract class TestPhaseStatus {
     componentRef?: any;
     hasWarning?: boolean;
     tags?: string[];
+    errorFingerprint?: string;
 
     private _name: string;
     private _start: Date;
@@ -586,7 +604,7 @@ export class MetadataRef {
     orderedKeys: string[];
     orderedValues: string[];
 
-    private static keysOrderPriority = ["Project", "BuildTarget", "Platform"]; // Should those be set from a config?
+    private static keysOrderPriority = ["Project", "Platform", "BuildTarget"]; // Should those be set from a config?
     private static valuesToIgnore = new Set(["default", "Default"]);
 
     static identifyCommonKeys(items: MetadataRef[]): string[] {
@@ -970,7 +988,7 @@ export class TestDataFilterSearchParams {
     }
 
     private encodeValue(items: string | string[]): string {
-        if (typeof items === "string") {
+        if (!Array.isArray(items)) {
             return items;
         }
 
@@ -1108,24 +1126,24 @@ export class TestDataHandler {
         }
     }
 
-    private checkSearchCallback(previous: string) {
+    private checkSearchCallback(previous: string): boolean {
         if (previous != this.search.toString()) {
             this.setCheckSearchUpdated();
+            return true;
         }
+        return false;
     }
 
-    setSearchParam(key: string, value: string, replace?: boolean) {
+    setSearchParam(key: string, value: string | string[], replace?: boolean) {
         const previous = this.search.toString();
         this.search.set(key, null, value);
-        replace && this.queueHistoryReplace();
-        this.checkSearchCallback(previous);
+        this.checkSearchCallback(previous) && replace && this.queueHistoryReplace();
     }
 
     removeSearchParam(key: string, replace?: boolean) {
         const previous = this.search.toString();
         this.search.delete(key);
-        replace && this.queueHistoryReplace();
-        this.checkSearchCallback(previous);
+        this.checkSearchCallback(previous) && replace && this.queueHistoryReplace();
     }
 
     getSearchParam(key: string): string | string[] | null | undefined {
@@ -1675,23 +1693,37 @@ export class TestDataHandler {
         }
     }
 
-    async queryPhase(phase: TestPhaseStatus, streamId: StreamId, isSubQuery: boolean = false) {
+    async queryPhaseSessions(phases: TestPhaseStatus[], streamId: StreamId, isSubQuery: boolean = false) {
 
-        const phaseKey = phase.key;
-        const test = phase.test!;
-        const testStatus = this.testStatus.get(streamId)?.tests.get(test);
+        if (phases.length === 0) {
+            return;
+        }
+
+        const test = phases[0].test!;
+        if (!phases.every(phase => phase.test === test)) {
+            console.error('Phase sessions query is not for the same test');
+            return;
+        }
 
-        const associatePhaseSession = () => {
-            if (!phase.session && testStatus) {
-                /// associate the corresponding phase session with the phase status
-                phase.associateSession(testStatus.getPhaseSessions(phaseKey));
-                /// associate the corresponding phase reference
-                phase.phaseRef = test.phases?.get(phaseKey);
+        const testStatus = this.testStatus.get(streamId)?.tests.get(test);
+        const phaseKeys = Array.from(new Set(phases.map(phase => phase.key)));
+
+        const associatePhaseSessions = () => {
+            if (testStatus) {
+                phases.forEach(phase => {
+                    if (!phase.session) {
+                        /// associate the corresponding phase session with the phase status
+                        phase.associateSession(testStatus.getPhaseSessions(phase.key, phase.testMetadata));
+                        /// associate the corresponding phase reference
+                        phase.phaseRef = test.phases?.get(phase.key);
+                    }
+                });
             }
         }
 
-        if (!testStatus || testStatus.hasPhaseSession(phaseKey)) {
-            associatePhaseSession();
+        // Check if all phases already have sessions
+        if (!testStatus || phaseKeys.every(phaseKey => testStatus.hasPhaseSession(phaseKey))) {
+            associatePhaseSessions();
             isSubQuery? this.setSubPhaseQueryLoading(false): this.setPhaseQueryLoading(false);
             return;
         }
@@ -1710,7 +1742,7 @@ export class TestDataHandler {
             if (fetcherVersion.version === test.version) {
                 console.debug(`Fetching from version ${fetcherVersion.version} test phase sessions`);
                 try {
-                    await fetcherVersion.fetchPhaseSessions(testStreamRef, testStatus, phaseKey, queryWeeks);
+                    await fetcherVersion.fetchPhaseSessions(testStreamRef, testStatus, phaseKeys, queryWeeks);
                 } catch (reason) {
                     /** ignore failure as getting the related phase sessions is optional to display content */
                     console.debug(`Failed to fetch from version ${fetcherVersion.version}: ${reason}`);
@@ -1719,16 +1751,42 @@ export class TestDataHandler {
             }
         }
 
-        /// sort phase session history by commit order descending
-        testStatus.sortPhaseSessions(phaseKey);
+        /// sort phase session history by commit order descending for all phases
+        phaseKeys.forEach(phaseKey => {
+            testStatus.sortPhaseSessions(phaseKey);
+        });
 
-        associatePhaseSession();
+        associatePhaseSessions();
 
         if (this.loaded) {
             isSubQuery? this.setSubPhaseQueryLoading(false): this.setPhaseQueryLoading(false);
         }
     }
 
+    async queryTestPhases(test: TestMetaStatus, phaseKeys?: string[]) {
+
+        if (!!phaseKeys && phaseKeys.every(phaseKey => test.test.phases?.get(phaseKey))) {
+            this.setPhaseQueryLoading(false);
+            return;
+        }
+
+        this.setPhaseQueryLoading(true);
+
+        for (const fetcherVersion of TestDataVersionRegistrar.manifest) {
+            /// fetch only from the matching version
+            console.debug(`Fetching from version ${fetcherVersion.version} test phases`);
+            try {
+                await fetcherVersion.fetchTestPhases(test, phaseKeys);
+                break;
+            } catch (reason) {
+                /** ignore failed attempt to try with next version */
+                console.debug(`Failed to fetch phases from version ${fetcherVersion.version}: ${reason}`);
+            }
+        }
+
+        this.setPhaseQueryLoading(false)
+    }
+
     async queryRecipes(test: TestNameRef, recipeIds: TestRecipeId[]) {
 
         const missingRecipes = recipeIds.filter(r => !test.getRecipe(r));
diff --git a/Engine/Source/Programs/Horde/HordeDashboard/plugins/test-automation/testDataV2Fetcher.ts b/Engine/Source/Programs/Horde/HordeDashboard/plugins/test-automation/testDataV2Fetcher.ts
--- a/Engine/Source/Programs/Horde/HordeDashboard/plugins/test-automation/testDataV2Fetcher.ts
+++ b/Engine/Source/Programs/Horde/HordeDashboard/plugins/test-automation/testDataV2Fetcher.ts
@@ -12,6 +12,7 @@ class TestPhaseStatusV2 extends TestPhaseStatus {
         super(phase.key, phase.name ?? "<undefinied>", new Date(phase.dateTime), phase.timeElapseSec, devices, phase.outcome, parent)
         this._eventStreamPath = phase.eventStreamPath;
         this.hasWarning = phase.hasWarning;
+        this.errorFingerprint = phase.errorFingerprint;
     }
 
     async fetchEventStream(): Promise<EventEntry[] | undefined> {
@@ -132,27 +133,38 @@ export class TestDataV2 implements TestVersionFetcher {
         return await getJobStepTestDataV2(jobId, stepId)
     }
 
-    async fetchPhaseSessions(streamRef: TestStreamRef, status: TestMetaStatus, phaseKey: string, weeks: number): Promise<boolean> {
+    async fetchPhaseSessions(streamRef: TestStreamRef, status: TestMetaStatus, phaseKeys: string[], weeks: number): Promise<boolean> {
 
         const streamId = streamRef.id;
         const tags: Map<string, TestTagRef> = new Map(streamRef.tags.filter((t) => t.version >= this.version).map((t) => [t.id, t]));
         const meta: Map<string, MetadataRef> = new Map(streamRef.meta.filter((m) => m.version >= this.version).map((m) => [m.id, m]));
 
-        if (!status.test.phases) {
-            const responses = await getTestPhasesV2([status.test.id]);
-            status.test.phases = new Map(responses.map(i => i.phases.map(p => this.createPhaseRef(p))).flat().map(p => [p.key, p]));
-        }
+        await this.fetchTestPhases(status, phaseKeys);        
+        
+        // Get phase refs for all requested phase keys and create a map from phase ID to phaseRef
+        const phaseIdToPhaseRefMap = new Map<string, TestPhaseRef>();
         
-        const phaseRef = status.test.phases.get(phaseKey);
-        if (phaseRef) {
+        phaseKeys.forEach(phaseKey => {
+            const phaseRef = status.test.phases?.get(phaseKey);
+            if (phaseRef) {
+                phaseIdToPhaseRefMap.set(phaseRef.id, phaseRef);
+            }
+        });
+
+        if (phaseIdToPhaseRefMap.size > 0) {
             const minQueryDate = new Date(new Date().valueOf() - (minuteInWeek * weeks * 60000));
             const maxQueryDate = new Date();
     
-            const responses = await getTestPhaseSessionsV2([streamId], [phaseRef.id], minQueryDate.toISOString(), maxQueryDate.toISOString());
+            const phaseIds = Array.from(phaseIdToPhaseRefMap.keys());
+            const responses = await getTestPhaseSessionsV2([streamId], phaseIds, minQueryDate.toISOString(), maxQueryDate.toISOString());
             responses.forEach((session) => {
                 if (session.streamId !== streamId) {
                     return;
                 }
+                const phaseRef = phaseIdToPhaseRefMap.get(session.phaseRef);
+                if (!phaseRef) {
+                    return;
+                }
                 const phaseSession = this.createPhaseSession(session);
                 const tagRefs = phaseSession.tagIds?.map((id) => tags.get(id)).filter((ref) => !!ref);
                 if (!!tagRefs) {
@@ -160,7 +172,7 @@ export class TestDataV2 implements TestVersionFetcher {
                 }
                 const metadataRef = meta.get(phaseSession.metadataId);
                 if (!!metadataRef) {
-                    status.addPhaseSession(metadataRef, phaseKey, phaseSession);
+                    status.addPhaseSession(metadataRef, phaseRef.key, phaseSession);
                 }
             });
 
@@ -171,6 +183,25 @@ export class TestDataV2 implements TestVersionFetcher {
 
     }
 
+    async fetchTestPhases(status: TestMetaStatus, phaseKeys?: string[]) {
+        const missingPhaseKeys = phaseKeys?.filter(p => !status.test.phases?.get(p));
+        if (!phaseKeys || missingPhaseKeys!.length > 0) {
+            const responses = await getTestPhasesV2([status.test.id]); /** TODO: filter by phaseKeys seems to be not working */
+            if (!status.test.phases) {
+                status.test.phases = new Map();
+            }
+            responses.forEach((test) => {
+                if (test.testId !== status.test.id) {
+                    return;
+                }
+                test.phases.forEach((phase) => {
+                    const phaseRef = this.createPhaseRef(phase);
+                    status.test.phases!.set(phaseRef.key, phaseRef);
+                });
+            });
+        }
+    }
+
     async fetchTestRecipes(testRef: TestNameRef, recipeIds: string[]) {
 
         const responses = await getTestRecipesV2(recipeIds);
diff --git a/Engine/Source/Programs/Horde/HordeDashboard/plugins/test-automation/testHistoryGraph.ts b/Engine/Source/Programs/Horde/HordeDashboard/plugins/test-automation/testHistoryGraph.ts
--- a/Engine/Source/Programs/Horde/HordeDashboard/plugins/test-automation/testHistoryGraph.ts
+++ b/Engine/Source/Programs/Horde/HordeDashboard/plugins/test-automation/testHistoryGraph.ts
@@ -246,7 +246,7 @@ export class TestHistoryGraph {
             if (items.length > 1) {
                 desc += renderToString(SessionListSelector(items, index, '(mouse wheel to select)'));
             }
-            if (!!cmeta || items.length > 1) desc += `#${items.length > 1? index + 1 : ''} ${cmeta} <br/>`;
+            if (!!cmeta || items.length > 1) desc += (items.length > 1? `::${(index + 1)} ` : '') + `${cmeta} <br/>`;
             desc += `Commit ${item.commitId} <br/>`;
             desc += `Duration ${msecToElapsed(item.duration * 1000, true, true)} <br/>`;
             desc += `${date} <br/>`;
