[Test Automation Hub] Fix filtering issue with Phase grid view

#rnx

#Synced-CL 51412642

diff --git a/Engine/Source/Programs/Horde/HordeDashboard/plugins/test-automation/testAutomationPhaseGridView.tsx b/Engine/Source/Programs/Horde/HordeDashboard/plugins/test-automation/testAutomationPhaseGridView.tsx
--- a/Engine/Source/Programs/Horde/HordeDashboard/plugins/test-automation/testAutomationPhaseGridView.tsx
+++ b/Engine/Source/Programs/Horde/HordeDashboard/plugins/test-automation/testAutomationPhaseGridView.tsx
@@ -625,6 +625,8 @@ export const TestPhasesGridView: React.FC<{ testId: string, handler: TestDataHan
                         const phaseGroups: PhaseStatusGroup[] = [];
                         sessionDetails.forEach(session => {
                             if (!session.meta) return;
+                            const metaItem = metaSelectorContext.metas.find(m => m.meta === session.meta);
+                            if (!metaItem) return;
                             session.phases.forEach(phase => {
                                 const phaseRef = test!.phases?.get(phase.key);
                                 if (!phaseRef) return;
@@ -639,7 +641,7 @@ export const TestPhasesGridView: React.FC<{ testId: string, handler: TestDataHan
                                     seen.set(phaseRef.key, group);
                                     phaseGroups.push(group);
                                 };
-                                group.addMetadata(metaSelectorContext.metas.find(m => m.meta === session.meta)!, phase);
+                                group.addMetadata(metaItem, phase);
                             });
                         });
                         metaSelectorContext.setPhaseGroups(phaseGroups);
@@ -662,7 +664,6 @@ export const TestPhasesGridView: React.FC<{ testId: string, handler: TestDataHan
         if (phaseIndex === -1) return;
         if (!isPhaseVisible(phaseIndex)) {
             // force the list to load the page that contains the target item
-            console.log('onPagesUpdated', metaSelectorContext.focusPhase);
             listRef.current?.scrollToIndex(phaseIndex, undefined, ScrollToMode.center);
         } else {
             // validate the scroll to index was achieved
