Horde: Use "Primary" in StorageController GC summary
#rnx

diff --git a/Engine/Source/Programs/Horde/Plugins/Storage/HordeServer.Storage/Storage/StorageController.cs b/Engine/Source/Programs/Horde/Plugins/Storage/HordeServer.Storage/Storage/StorageController.cs
--- a/Engine/Source/Programs/Horde/Plugins/Storage/HordeServer.Storage/Storage/StorageController.cs
+++ b/Engine/Source/Programs/Horde/Plugins/Storage/HordeServer.Storage/Storage/StorageController.cs
@@ -873,7 +873,7 @@ namespace HordeServer.Storage
 	/// </summary>
 	public class StorageDashboardStatus
 	{
-		/// <summary>Master GC switch</summary>
+		/// <summary>Primary GC switch</summary>
 		public bool EnableGc { get; set; }
 
 		/// <summary>Verification mode (log-only, no deletes)</summary>
