[Test Automation Hub] Consider errorfingerprint repro days for each individual meta

#rnx

#Synced-CL 51412642

diff --git a/Engine/Source/Programs/Horde/HordeDashboard/plugins/test-automation/phaseGrid.ts b/Engine/Source/Programs/Horde/HordeDashboard/plugins/test-automation/phaseGrid.ts
--- a/Engine/Source/Programs/Horde/HordeDashboard/plugins/test-automation/phaseGrid.ts
+++ b/Engine/Source/Programs/Horde/HordeDashboard/plugins/test-automation/phaseGrid.ts
@@ -482,7 +482,7 @@ export class MetaSelectorContext {
         const testStatus = MetaSelectorContext.getTestStatus(this._handler, this.test);
         group.metadata.forEach((status, metaItem) => {
             if (status?.errorFingerprint) {
-                let startDate = errorsCaught.get(status.errorFingerprint) ?? status.session?.start ?? new Date();
+                let startDate = status.session?.start ?? new Date();
                 const metaStatus = testStatus?.getPhaseSessions(group.phase.key, metaItem.meta);
                 let consecutivePasses = 0;
                 metaStatus?.some(s => {
@@ -496,7 +496,7 @@ export class MetaSelectorContext {
                     consecutivePasses = 0;
                     return false;
                 });
-                errorsCaught.set(status.errorFingerprint, startDate);
+                errorsCaught.set(metaItem.meta.id, startDate);
             }
         });
         return errorsCaught;
diff --git a/Engine/Source/Programs/Horde/HordeDashboard/plugins/test-automation/testAutomationPhaseGridView.tsx b/Engine/Source/Programs/Horde/HordeDashboard/plugins/test-automation/testAutomationPhaseGridView.tsx
--- a/Engine/Source/Programs/Horde/HordeDashboard/plugins/test-automation/testAutomationPhaseGridView.tsx
+++ b/Engine/Source/Programs/Horde/HordeDashboard/plugins/test-automation/testAutomationPhaseGridView.tsx
@@ -57,7 +57,7 @@ const PhaseStatusWidget: React.FC<{ group: PhaseStatusGroup, context: MetaSelect
                 { context.metas.filter(m => m.visible).map(m => {
                     const status = group.getMetadataStatus(m);
                     const iconColor = status ? getPhaseSessionStatusColor(status) : undefined;
-                    const errorDays = errorsCaught?.get(status?.errorFingerprint ?? "") ?? 0;
+                    const errorDays = errorsCaught?.get(m.meta.id) ?? 0;
                     const title = `${status ? status.outcome : 'N/A'}${!!m.name ? ` - ${m.name}` : ''}`;
                     return <Stack key={`meta-${m.meta.id}`}
                                 className={`${getMetaCellClassName(m.meta.id)} ${getPhaseCellClassName(group.phase.key)}`}
@@ -67,7 +67,7 @@ const PhaseStatusWidget: React.FC<{ group: PhaseStatusGroup, context: MetaSelect
                                 onClick={() => onClick(group.phase.key, m)}
                                 >                                    
                                     { !!status && <FontIcon style={{ fontSize: 12 * (errorDays/context.daysQueried + 1), color: iconColor }} iconName="Square" /> }
-                                    { !!status && !!status.errorFingerprint && !!errorsCaught && errorDays > 0 &&
+                                    { !!status && !!status.errorFingerprint && errorDays > 0 &&
                                         <Text style={{ fontSize: 8, color: 'white', position: 'absolute', cursor: 'pointer', userSelect: "none" }}>{errorDays >= context.daysQueried ? `${context.daysQueried}+` : errorDays}</Text>
                                     }
                                     { !!status && !!status.errorFingerprint &&
