Remove mock collection stats data from StorageView, use live server data
#rnx

diff --git a/Engine/Source/Programs/Horde/HordeDashboard/src/components/storage/StorageView.tsx b/Engine/Source/Programs/Horde/HordeDashboard/src/components/storage/StorageView.tsx
--- a/Engine/Source/Programs/Horde/HordeDashboard/src/components/storage/StorageView.tsx
+++ b/Engine/Source/Programs/Horde/HordeDashboard/src/components/storage/StorageView.tsx
@@ -85,21 +85,6 @@ class StorageHandler extends PollBase {
 
 const handler = new StorageHandler();
 
-// ─── Mock Collection Stats (from prod collStats query 2026-03-04 — remove once live data flows) ───
-const mockCollectionStats: CollectionStatsSnapshot[] = [
-    {
-        time: "2026-03-04T20:00:00Z",
-        collections: [
-            { name: "Storage.Blobs", documentCount: 432658733, dataSizeBytes: 369057899249, indexSizeBytes: 2778013302784, storageSizeBytes: 1927312744448, avgDocSizeBytes: 853, fragmentationRatio: 5.22 },
-            { name: "Storage.Refs", documentCount: 153234679, dataSizeBytes: 42292771404, indexSizeBytes: 45889093632, storageSizeBytes: 49007067136, avgDocSizeBytes: 276, fragmentationRatio: 1.16 },
-            { name: "ArtifactsV2", documentCount: 5120964, dataSizeBytes: 2770441524, indexSizeBytes: 11259772928, storageSizeBytes: 3079438336, avgDocSizeBytes: 541, fragmentationRatio: 1.11 },
-            { name: "Jobs", documentCount: 3778876, dataSizeBytes: 32668383020, indexSizeBytes: 62730780672, storageSizeBytes: 77837942784, avgDocSizeBytes: 8645, fragmentationRatio: 2.38 },
-            { name: "LogFiles", documentCount: 441678629, dataSizeBytes: 1175748510398, indexSizeBytes: 38212583424, storageSizeBytes: 295902371840, avgDocSizeBytes: 2662, fragmentationRatio: 0.25 },
-            { name: "LogEvents", documentCount: 335679442, dataSizeBytes: 36253379736, indexSizeBytes: 109965975552, storageSizeBytes: 44516483072, avgDocSizeBytes: 108, fragmentationRatio: 1.23 },
-        ]
-    }
-];
-
 // ─── Helpers ───
 
 function formatCount(n: number): string {
@@ -837,7 +822,7 @@ const StoragePanel: React.FC = observer(() => {
             {handler.data && (
                 <>
                     <GcHealthPanel entries={handler.data.gcMetrics.entries} status={handler.data.status} />
-                    <CollectionHealthPanel snapshots={mockCollectionStats} />
+                    <CollectionHealthPanel snapshots={handler.data.collectionStats.entries} />
                     <GcRatePanel entries={handler.data.gcMetrics.entries} />
                     <GcConfigPanel status={handler.data.status} />
                     <NamespaceStatusPanel status={handler.data.status} />
