Storage Dashboard: add Ingested stat to GC Health Summary
#rnx

diff --git a/Engine/Source/Programs/Horde/HordeDashboard/src/components/storage/StorageView.tsx b/Engine/Source/Programs/Horde/HordeDashboard/src/components/storage/StorageView.tsx
--- a/Engine/Source/Programs/Horde/HordeDashboard/src/components/storage/StorageView.tsx
+++ b/Engine/Source/Programs/Horde/HordeDashboard/src/components/storage/StorageView.tsx
@@ -374,7 +374,9 @@ const GcHealthPanel: React.FC<{ entries: StorageGcMetricEntry[]; status: Storage
             <Stack horizontal tokens={{ childrenGap: 8 }} style={{ marginBottom: 12 }}>
                 {statItem("Ingestion", `~${formatCount(ingestionRate)}/hr`, undefined, 130)}
                 {statItem("Deletion", `~${formatCount(deletionRate)}/hr`, undefined, 130)}
-                {statItem("Freed", formatBytes(totalBytesFreed), undefined, 120)}
+                {statItem("Ingested", `${formatCount(totalIngested)} blobs`, "#ca5010", 130)}
+                {statItem("Deleted", `${formatCount(totalDeleted)} blobs`, "#107c10", 130)}
+                {statItem("Freed", formatBytes(totalBytesFreed), "#107c10", 110)}
                 {statItem("Paused NS", `${pausedCount} / ${status.namespaces.length}`, pausedCount > 0 ? "#d13438" : undefined, 100)}
                 {statItem("Enqueue Fail", formatCount(totalEnqueueFails), totalEnqueueFails > 0 ? "#d13438" : undefined, 120)}
                 {statItem("Throttles", formatCount(totalThrottle), totalThrottle > 0 ? "#fce100" : undefined, 100)}
