Horde: Group crash callstack lines into single events in log matchers
LogChannelEventMatcher: Detect crash headers and [Callstack] sequences, consuming the entire crash block as one multi-line event instead of matching each line individually. Extract assertion/ensure condition as a CrashSummary property for platform-agnostic fingerprinting.
HashedIssueHandler: When CrashSummary is present, hash on the summary instead of the full callstack render, enabling cross-platform issue matching.
CrashEventMatcher: Fix intermittent failure to match begin/end: stack for UAT blocks caused by incomplete buffer look-ahead.
Consume trailing blank error lines appended as padding after crash output.
Add matcher and issue handler tests for crash grouping, fingerprint uniqueness, and backward compatibility.

#Synced-CL 0

diff --git a/Engine/Source/Programs/Horde/Plugins/Build/HordeServer.Build.Tests/Issues/IssueHandlerTests.cs b/Engine/Source/Programs/Horde/Plugins/Build/HordeServer.Build.Tests/Issues/IssueHandlerTests.cs
--- a/Engine/Source/Programs/Horde/Plugins/Build/HordeServer.Build.Tests/Issues/IssueHandlerTests.cs
+++ b/Engine/Source/Programs/Horde/Plugins/Build/HordeServer.Build.Tests/Issues/IssueHandlerTests.cs
@@ -226,6 +226,276 @@ namespace HordeServer.Tests.Issues
 			Assert.AreEqual("LineSetComponentMaterial", groupedIssues[2].Metadata.FindValues("Material").First());
 		}
 
+		// Parse structured logs where some entries represent multi-line events (multiple JSON lines sharing lineCount > 1).
+		// Lines with lineCount > 1 are grouped into a single IssueEvent, simulating matcher output after grouping.
+		static IssueHandlerTestData ParseGroupedJSONStructuredLogs(string[] jsonStructuredLogs)
+		{
+			List<LogEvent> logEvents = new List<LogEvent>();
+			List<IssueEvent> issueEvents = new List<IssueEvent>();
+			List<JsonLogEvent> pendingLines = new List<JsonLogEvent>();
+			int pendingLineCount = 0;
+			int pendingLineIndex = 0;
+			Microsoft.Extensions.Logging.LogLevel pendingLevel = Microsoft.Extensions.Logging.LogLevel.None;
+			Microsoft.Extensions.Logging.EventId? pendingEventId = null;
+
+			foreach (string jsonStructuredLog in jsonStructuredLogs)
+			{
+				byte[] byteArrayStructuredLog = Encoding.UTF8.GetBytes(jsonStructuredLog);
+				logEvents.Add(LogEvent.Read(byteArrayStructuredLog));
+
+				JsonLogEvent jsonLogEvent;
+				Assert.IsTrue(JsonLogEvent.TryParse(byteArrayStructuredLog, out jsonLogEvent));
+
+				if (jsonLogEvent.LineCount > 1)
+				{
+					// Multi-line event: accumulate lines
+					if (pendingLines.Count == 0)
+					{
+						pendingLineCount = jsonLogEvent.LineCount;
+						pendingLineIndex = jsonLogEvent.LineIndex;
+						pendingLevel = jsonLogEvent.Level;
+						pendingEventId = jsonLogEvent.EventId;
+					}
+					pendingLines.Add(jsonLogEvent);
+
+					if (pendingLines.Count >= pendingLineCount)
+					{
+						// All lines collected, create a single IssueEvent
+						issueEvents.Add(new IssueEvent(pendingLineIndex, pendingLevel, pendingEventId, new List<JsonLogEvent>(pendingLines)));
+						pendingLines.Clear();
+					}
+				}
+				else
+				{
+					// Single-line event
+					List<JsonLogEvent> singleList = new List<JsonLogEvent> { jsonLogEvent };
+					issueEvents.Add(new IssueEvent(jsonLogEvent.LineIndex, jsonLogEvent.Level, jsonLogEvent.EventId, singleList));
+				}
+			}
+
+			Assert.AreEqual(0, pendingLines.Count, "Incomplete multi-line group in test data: lineCount indicated more lines than were provided");
+			return new IssueHandlerTestData(logEvents, issueEvents);
+		}
+
+		// Test HashedIssueHandler with a single IssueHandler using grouped logs
+		static List<IssueEventGroup> TestSingleIssueHandlerGrouped(IssueHandler issueHandler, string[] jsonStructuredLogs)
+		{
+			IssueHandlerTestData issueHandlerParsedTestData = ParseGroupedJSONStructuredLogs(jsonStructuredLogs);
+			foreach (IssueEvent issueEvent in issueHandlerParsedTestData.IssueEvents)
+			{
+				issueHandler.HandleEvent(issueEvent);
+			}
+			List<IssueEventGroup> groupedIssues = issueHandler.GetIssues().ToList();
+			return groupedIssues;
+		}
+
+		[TestMethod]
+		public void HashedIssueHandler_CrashCallstackGroupedIntoOneEvent()
+		{
+			// Simulates a crash sequence after the matcher groups it into a single multi-line event.
+			// All lines share lineCount=8 and have incrementing line indices.
+			// The first line has the Callstack property set to true.
+			string[] jsonStructuredLogs =
+			{
+				"{ \"time\":\"2025-01-01T01:00:00.000Z\",\"level\":\"Error\",\"message\":\"LogWindows: Error: === Critical error: ===\",\"id\":301,\"line\":0,\"lineCount\":8,\"properties\":{ \"_channel\":{ \"$type\":\"Channel\",\"$text\":\"LogWindows\"},\"_severity\":{ \"$type\":\"Severity\",\"$text\":\"Error\"},\"Callstack\":true} }",
+				"{ \"time\":\"2025-01-01T01:00:00.001Z\",\"level\":\"Error\",\"message\":\"LogWindows: Error: Assertion failed: Pair != nullptr [File:Map.h] [Line: 728]\",\"id\":301,\"line\":1,\"lineCount\":8,\"properties\":{ \"_channel\":{ \"$type\":\"Channel\",\"$text\":\"LogWindows\"},\"_severity\":{ \"$type\":\"Severity\",\"$text\":\"Error\"}} }",
+				"{ \"time\":\"2025-01-01T01:00:00.002Z\",\"level\":\"Error\",\"message\":\"LogWindows: Error: [Callstack] 0x00007ffe5b17b248 UnrealEditor-Core.dll!FDebug::CheckVerifyFailedImpl2() []\",\"id\":301,\"line\":2,\"lineCount\":8,\"properties\":{ \"_channel\":{ \"$type\":\"Channel\",\"$text\":\"LogWindows\"},\"_severity\":{ \"$type\":\"Severity\",\"$text\":\"Error\"}} }",
+				"{ \"time\":\"2025-01-01T01:00:00.003Z\",\"level\":\"Error\",\"message\":\"LogWindows: Error: [Callstack] 0x00007ffe589a608c UnrealEditor-Engine.dll!UWorld::Tick() []\",\"id\":301,\"line\":3,\"lineCount\":8,\"properties\":{ \"_channel\":{ \"$type\":\"Channel\",\"$text\":\"LogWindows\"},\"_severity\":{ \"$type\":\"Severity\",\"$text\":\"Error\"}} }",
+				"{ \"time\":\"2025-01-01T01:00:00.004Z\",\"level\":\"Error\",\"message\":\"LogWindows: Error: [Callstack] 0x00007ffe589b1234 UnrealEditor-Engine.dll!FEngineLoop::Tick() []\",\"id\":301,\"line\":4,\"lineCount\":8,\"properties\":{ \"_channel\":{ \"$type\":\"Channel\",\"$text\":\"LogWindows\"},\"_severity\":{ \"$type\":\"Severity\",\"$text\":\"Error\"}} }",
+				"{ \"time\":\"2025-01-01T01:00:00.005Z\",\"level\":\"Error\",\"message\":\"LogWindows: Error: [Callstack] 0x00007ffe589c5678 UnrealEditor-Cmd.exe!GuardedMain() []\",\"id\":301,\"line\":5,\"lineCount\":8,\"properties\":{ \"_channel\":{ \"$type\":\"Channel\",\"$text\":\"LogWindows\"},\"_severity\":{ \"$type\":\"Severity\",\"$text\":\"Error\"}} }",
+				"{ \"time\":\"2025-01-01T01:00:00.006Z\",\"level\":\"Error\",\"message\":\"LogWindows: Error: [Callstack] 0x00007ffe89f44cb0 KERNEL32.DLL!UnknownFunction []\",\"id\":301,\"line\":6,\"lineCount\":8,\"properties\":{ \"_channel\":{ \"$type\":\"Channel\",\"$text\":\"LogWindows\"},\"_severity\":{ \"$type\":\"Severity\",\"$text\":\"Error\"}} }",
+				"{ \"time\":\"2025-01-01T01:00:00.007Z\",\"level\":\"Error\",\"message\":\"LogWindows: Error: [Callstack] 0x00007ffe9d14a271 ntdll.dll!UnknownFunction []\",\"id\":301,\"line\":7,\"lineCount\":8,\"properties\":{ \"_channel\":{ \"$type\":\"Channel\",\"$text\":\"LogWindows\"},\"_severity\":{ \"$type\":\"Severity\",\"$text\":\"Error\"}} }",
+			};
+
+			IssueHandlerContext context = new IssueHandlerContext(new StreamId("Release-XX.XX"), new TemplateId(), "Compile ProjectX", new Dictionary<string, string>(), new Dictionary<string, string>());
+			HashedIssueHandler handler = new HashedIssueHandler(context);
+			List<IssueEventGroup> groupedIssues = TestSingleIssueHandlerGrouped(handler, jsonStructuredLogs);
+
+			// All 8 lines form a single multi-line IssueEvent, which the handler hashes as one group
+			Assert.AreEqual(1, groupedIssues.Count);
+			Assert.AreEqual(1, groupedIssues[0].Events.Count);
+		}
+
+		[TestMethod]
+		public void HashedIssueHandler_RegularLogChannelEventsRemainSeparate()
+		{
+			// Regular LogChannel errors without Callstack property should each produce a separate group
+			string[] jsonStructuredLogs =
+			{
+				"{ \"time\":\"2025-01-01T01:00:00.000Z\",\"level\":\"Error\",\"message\":\"LogAudio: Error: Failed to initialize audio device\",\"id\":301,\"properties\":{ \"_channel\":{ \"$type\":\"Channel\",\"$text\":\"LogAudio\"},\"_severity\":{ \"$type\":\"Severity\",\"$text\":\"Error\"}} }",
+				"{ \"time\":\"2025-01-01T01:00:00.001Z\",\"level\":\"Error\",\"message\":\"LogRenderer: Error: Shader compilation failed\",\"id\":301,\"properties\":{ \"_channel\":{ \"$type\":\"Channel\",\"$text\":\"LogRenderer\"},\"_severity\":{ \"$type\":\"Severity\",\"$text\":\"Error\"}} }",
+			};
+
+			IssueHandlerContext context = new IssueHandlerContext(new StreamId("Release-XX.XX"), new TemplateId(), "Compile ProjectX", new Dictionary<string, string>(), new Dictionary<string, string>());
+			HashedIssueHandler handler = new HashedIssueHandler(context);
+			List<IssueEventGroup> groupedIssues = TestSingleIssueHandler(handler, jsonStructuredLogs);
+
+			// Each regular event gets its own group with a unique hash
+			Assert.AreEqual(2, groupedIssues.Count);
+			Assert.AreEqual(1, groupedIssues[0].Events.Count);
+			Assert.AreEqual(1, groupedIssues[1].Events.Count);
+		}
+
+		[TestMethod]
+		public void HashedIssueHandler_DifferentCrashesProduceDifferentHashes()
+		{
+			// Two different crashes with different callstacks should produce different hash keys
+			string[] crashA =
+			{
+				"{ \"time\":\"2025-01-01T01:00:00.000Z\",\"level\":\"Error\",\"message\":\"LogWindows: Error: === Critical error: ===\",\"id\":301,\"line\":0,\"lineCount\":3,\"properties\":{ \"_channel\":{ \"$type\":\"Channel\",\"$text\":\"LogWindows\"},\"_severity\":{ \"$type\":\"Severity\",\"$text\":\"Error\"},\"Callstack\":true} }",
+				"{ \"time\":\"2025-01-01T01:00:00.001Z\",\"level\":\"Error\",\"message\":\"LogWindows: Error: [Callstack] 0x00007ffe5b17b248 UnrealEditor-Core.dll!FDebug::CheckVerifyFailedImpl2() []\",\"id\":301,\"line\":1,\"lineCount\":3,\"properties\":{ \"_channel\":{ \"$type\":\"Channel\",\"$text\":\"LogWindows\"},\"_severity\":{ \"$type\":\"Severity\",\"$text\":\"Error\"}} }",
+				"{ \"time\":\"2025-01-01T01:00:00.002Z\",\"level\":\"Error\",\"message\":\"LogWindows: Error: [Callstack] 0x00007ffe589a608c UnrealEditor-Engine.dll!UWorld::Tick() []\",\"id\":301,\"line\":2,\"lineCount\":3,\"properties\":{ \"_channel\":{ \"$type\":\"Channel\",\"$text\":\"LogWindows\"},\"_severity\":{ \"$type\":\"Severity\",\"$text\":\"Error\"}} }",
+			};
+
+			string[] crashB =
+			{
+				"{ \"time\":\"2025-01-01T01:00:00.000Z\",\"level\":\"Error\",\"message\":\"LogWindows: Error: === Critical error: ===\",\"id\":301,\"line\":0,\"lineCount\":3,\"properties\":{ \"_channel\":{ \"$type\":\"Channel\",\"$text\":\"LogWindows\"},\"_severity\":{ \"$type\":\"Severity\",\"$text\":\"Error\"},\"Callstack\":true} }",
+				"{ \"time\":\"2025-01-01T01:00:00.001Z\",\"level\":\"Error\",\"message\":\"LogWindows: Error: [Callstack] 0x00007ffe1234abcd UnrealEditor-Slate.dll!SWidget::Paint() []\",\"id\":301,\"line\":1,\"lineCount\":3,\"properties\":{ \"_channel\":{ \"$type\":\"Channel\",\"$text\":\"LogWindows\"},\"_severity\":{ \"$type\":\"Severity\",\"$text\":\"Error\"}} }",
+				"{ \"time\":\"2025-01-01T01:00:00.002Z\",\"level\":\"Error\",\"message\":\"LogWindows: Error: [Callstack] 0x00007ffe5678dcba UnrealEditor-SlateCore.dll!FSlateApplication::Tick() []\",\"id\":301,\"line\":2,\"lineCount\":3,\"properties\":{ \"_channel\":{ \"$type\":\"Channel\",\"$text\":\"LogWindows\"},\"_severity\":{ \"$type\":\"Severity\",\"$text\":\"Error\"}} }",
+			};
+
+			IssueHandlerContext context = new IssueHandlerContext(new StreamId("Release-XX.XX"), new TemplateId(), "Compile ProjectX", new Dictionary<string, string>(), new Dictionary<string, string>());
+
+			HashedIssueHandler handlerA = new HashedIssueHandler(context);
+			List<IssueEventGroup> groupsA = TestSingleIssueHandlerGrouped(handlerA, crashA);
+
+			HashedIssueHandler handlerB = new HashedIssueHandler(context);
+			List<IssueEventGroup> groupsB = TestSingleIssueHandlerGrouped(handlerB, crashB);
+
+			Assert.AreEqual(1, groupsA.Count);
+			Assert.AreEqual(1, groupsB.Count);
+
+			// The two crashes should have different keys (different callstack content)
+			Assert.IsFalse(groupsA[0].Keys.SetEquals(groupsB[0].Keys), "Different crashes should produce different hash keys");
+		}
+
+		[TestMethod]
+		public void HashedIssueHandler_SharedGenericFramesDontCauseFalseMatch()
+		{
+			// Two crashes that share a generic frame (KERNEL32.DLL!UnknownFunction) but differ
+			// in other frames should produce different fingerprints
+			string[] crashA =
+			{
+				"{ \"time\":\"2025-01-01T01:00:00.000Z\",\"level\":\"Error\",\"message\":\"LogWindows: Error: === Critical error: ===\",\"id\":301,\"line\":0,\"lineCount\":3,\"properties\":{ \"_channel\":{ \"$type\":\"Channel\",\"$text\":\"LogWindows\"},\"_severity\":{ \"$type\":\"Severity\",\"$text\":\"Error\"},\"Callstack\":true} }",
+				"{ \"time\":\"2025-01-01T01:00:00.001Z\",\"level\":\"Error\",\"message\":\"LogWindows: Error: [Callstack] 0x00007ffe5b17b248 UnrealEditor-Core.dll!FDebug::CheckVerifyFailedImpl2() []\",\"id\":301,\"line\":1,\"lineCount\":3,\"properties\":{ \"_channel\":{ \"$type\":\"Channel\",\"$text\":\"LogWindows\"},\"_severity\":{ \"$type\":\"Severity\",\"$text\":\"Error\"}} }",
+				"{ \"time\":\"2025-01-01T01:00:00.002Z\",\"level\":\"Error\",\"message\":\"LogWindows: Error: [Callstack] 0x00007ffe89f44cb0 KERNEL32.DLL!UnknownFunction []\",\"id\":301,\"line\":2,\"lineCount\":3,\"properties\":{ \"_channel\":{ \"$type\":\"Channel\",\"$text\":\"LogWindows\"},\"_severity\":{ \"$type\":\"Severity\",\"$text\":\"Error\"}} }",
+			};
+
+			string[] crashB =
+			{
+				"{ \"time\":\"2025-01-01T01:00:00.000Z\",\"level\":\"Error\",\"message\":\"LogWindows: Error: === Critical error: ===\",\"id\":301,\"line\":0,\"lineCount\":3,\"properties\":{ \"_channel\":{ \"$type\":\"Channel\",\"$text\":\"LogWindows\"},\"_severity\":{ \"$type\":\"Severity\",\"$text\":\"Error\"},\"Callstack\":true} }",
+				"{ \"time\":\"2025-01-01T01:00:00.001Z\",\"level\":\"Error\",\"message\":\"LogWindows: Error: [Callstack] 0x00007ffe1234abcd UnrealEditor-Slate.dll!SWidget::Paint() []\",\"id\":301,\"line\":1,\"lineCount\":3,\"properties\":{ \"_channel\":{ \"$type\":\"Channel\",\"$text\":\"LogWindows\"},\"_severity\":{ \"$type\":\"Severity\",\"$text\":\"Error\"}} }",
+				"{ \"time\":\"2025-01-01T01:00:00.002Z\",\"level\":\"Error\",\"message\":\"LogWindows: Error: [Callstack] 0x00007ffe89f44cb0 KERNEL32.DLL!UnknownFunction []\",\"id\":301,\"line\":2,\"lineCount\":3,\"properties\":{ \"_channel\":{ \"$type\":\"Channel\",\"$text\":\"LogWindows\"},\"_severity\":{ \"$type\":\"Severity\",\"$text\":\"Error\"}} }",
+			};
+
+			IssueHandlerContext context = new IssueHandlerContext(new StreamId("Release-XX.XX"), new TemplateId(), "Compile ProjectX", new Dictionary<string, string>(), new Dictionary<string, string>());
+
+			HashedIssueHandler handlerA = new HashedIssueHandler(context);
+			List<IssueEventGroup> groupsA = TestSingleIssueHandlerGrouped(handlerA, crashA);
+
+			HashedIssueHandler handlerB = new HashedIssueHandler(context);
+			List<IssueEventGroup> groupsB = TestSingleIssueHandlerGrouped(handlerB, crashB);
+
+			Assert.AreEqual(1, groupsA.Count);
+			Assert.AreEqual(1, groupsB.Count);
+
+			// Despite sharing KERNEL32.DLL!UnknownFunction, different unique frames mean different hashes
+			Assert.IsFalse(groupsA[0].Keys.SetEquals(groupsB[0].Keys), "Crashes sharing only generic frames should produce different hash keys");
+		}
+
+		[TestMethod]
+		public void HashedIssueHandler_CrashSummaryMatchesAcrossPlatforms()
+		{
+			// Same CrashSummary but different callstack module names (dll vs so) should produce the same hash
+			string[] windowsCrash =
+			{
+				"{ \"time\":\"2025-01-01T01:00:00.000Z\",\"level\":\"Error\",\"message\":\"LogWindows: Error: === Handled ensure: ===\",\"id\":301,\"line\":0,\"lineCount\":3,\"properties\":{ \"_channel\":{ \"$type\":\"Channel\",\"$text\":\"LogWindows\"},\"_severity\":{ \"$type\":\"Severity\",\"$text\":\"Error\"},\"Callstack\":true,\"CrashSummary\":\"Ensure condition failed: InParentOp [RunIKRigOp.cpp]\"} }",
+				"{ \"time\":\"2025-01-01T01:00:00.001Z\",\"level\":\"Error\",\"message\":\"LogWindows: Error: [Callstack] 0x00007fff5b17b248 UnrealEditor-IKRig.dll!FIKRetargetRunIKRigOp::ValidateIKChain() []\",\"id\":301,\"line\":1,\"lineCount\":3,\"properties\":{ \"_channel\":{ \"$type\":\"Channel\",\"$text\":\"LogWindows\"},\"_severity\":{ \"$type\":\"Severity\",\"$text\":\"Error\"}} }",
+				"{ \"time\":\"2025-01-01T01:00:00.002Z\",\"level\":\"Error\",\"message\":\"LogWindows: Error: [Callstack] 0x00007fff89f44cb0 KERNEL32.DLL!UnknownFunction []\",\"id\":301,\"line\":2,\"lineCount\":3,\"properties\":{ \"_channel\":{ \"$type\":\"Channel\",\"$text\":\"LogWindows\"},\"_severity\":{ \"$type\":\"Severity\",\"$text\":\"Error\"}} }",
+			};
+
+			string[] linuxCrash =
+			{
+				"{ \"time\":\"2025-01-01T01:00:00.000Z\",\"level\":\"Error\",\"message\":\"LogOutputDevice: Error: === Handled ensure: ===\",\"id\":301,\"line\":0,\"lineCount\":3,\"properties\":{ \"_channel\":{ \"$type\":\"Channel\",\"$text\":\"LogOutputDevice\"},\"_severity\":{ \"$type\":\"Severity\",\"$text\":\"Error\"},\"Callstack\":true,\"CrashSummary\":\"Ensure condition failed: InParentOp [RunIKRigOp.cpp]\"} }",
+				"{ \"time\":\"2025-01-01T01:00:00.001Z\",\"level\":\"Error\",\"message\":\"LogOutputDevice: Error: [Callstack] 0x00007c7b053585da libUnrealEditor-IKRig.so!FIKRetargetRunIKRigOp::ValidateIKChain() [/mnt/horde/Sync/Engine/Plugins/IKRig/RunIKRigOp.cpp:436]\",\"id\":301,\"line\":1,\"lineCount\":3,\"properties\":{ \"_channel\":{ \"$type\":\"Channel\",\"$text\":\"LogOutputDevice\"},\"_severity\":{ \"$type\":\"Severity\",\"$text\":\"Error\"}} }",
+				"{ \"time\":\"2025-01-01T01:00:00.002Z\",\"level\":\"Error\",\"message\":\"LogOutputDevice: Error: [Callstack] 0x00007c7ce429af1d libUnrealEditor-UnixCommonStartup.so!CommonUnixMain() []\",\"id\":301,\"line\":2,\"lineCount\":3,\"properties\":{ \"_channel\":{ \"$type\":\"Channel\",\"$text\":\"LogOutputDevice\"},\"_severity\":{ \"$type\":\"Severity\",\"$text\":\"Error\"}} }",
+			};
+
+			IssueHandlerContext context = new IssueHandlerContext(new StreamId("Release-XX.XX"), new TemplateId(), "Compile ProjectX", new Dictionary<string, string>(), new Dictionary<string, string>());
+
+			HashedIssueHandler handlerWin = new HashedIssueHandler(context);
+			List<IssueEventGroup> groupsWin = TestSingleIssueHandlerGrouped(handlerWin, windowsCrash);
+
+			HashedIssueHandler handlerLinux = new HashedIssueHandler(context);
+			List<IssueEventGroup> groupsLinux = TestSingleIssueHandlerGrouped(handlerLinux, linuxCrash);
+
+			Assert.AreEqual(1, groupsWin.Count);
+			Assert.AreEqual(1, groupsLinux.Count);
+
+			// Same CrashSummary should produce the same hash regardless of platform-specific callstack content
+			Assert.IsTrue(groupsWin[0].Keys.SetEquals(groupsLinux[0].Keys), "Same crash on different platforms should produce matching hash keys");
+		}
+
+		[TestMethod]
+		public void HashedIssueHandler_DifferentCrashSummariesProduceDifferentHashes()
+		{
+			// Two crashes with different CrashSummary values should not match
+			string[] crashA =
+			{
+				"{ \"time\":\"2025-01-01T01:00:00.000Z\",\"level\":\"Error\",\"message\":\"LogWindows: Error: === Critical error: ===\",\"id\":301,\"line\":0,\"lineCount\":2,\"properties\":{ \"_channel\":{ \"$type\":\"Channel\",\"$text\":\"LogWindows\"},\"_severity\":{ \"$type\":\"Severity\",\"$text\":\"Error\"},\"Callstack\":true,\"CrashSummary\":\"Assertion failed: Pair != nullptr [Map.h]\"} }",
+				"{ \"time\":\"2025-01-01T01:00:00.001Z\",\"level\":\"Error\",\"message\":\"LogWindows: Error: [Callstack] 0x00007ffe5b17b248 UnrealEditor-Core.dll!FDebug::CheckVerifyFailedImpl2() []\",\"id\":301,\"line\":1,\"lineCount\":2,\"properties\":{ \"_channel\":{ \"$type\":\"Channel\",\"$text\":\"LogWindows\"},\"_severity\":{ \"$type\":\"Severity\",\"$text\":\"Error\"}} }",
+			};
+
+			string[] crashB =
+			{
+				"{ \"time\":\"2025-01-01T01:00:00.000Z\",\"level\":\"Error\",\"message\":\"LogWindows: Error: === Handled ensure: ===\",\"id\":301,\"line\":0,\"lineCount\":2,\"properties\":{ \"_channel\":{ \"$type\":\"Channel\",\"$text\":\"LogWindows\"},\"_severity\":{ \"$type\":\"Severity\",\"$text\":\"Error\"},\"Callstack\":true,\"CrashSummary\":\"Ensure condition failed: InParentOp [RunIKRigOp.cpp]\"} }",
+				"{ \"time\":\"2025-01-01T01:00:00.001Z\",\"level\":\"Error\",\"message\":\"LogWindows: Error: [Callstack] 0x00007ffe5b17b248 UnrealEditor-Core.dll!FDebug::CheckVerifyFailedImpl2() []\",\"id\":301,\"line\":1,\"lineCount\":2,\"properties\":{ \"_channel\":{ \"$type\":\"Channel\",\"$text\":\"LogWindows\"},\"_severity\":{ \"$type\":\"Severity\",\"$text\":\"Error\"}} }",
+			};
+
+			IssueHandlerContext context = new IssueHandlerContext(new StreamId("Release-XX.XX"), new TemplateId(), "Compile ProjectX", new Dictionary<string, string>(), new Dictionary<string, string>());
+
+			HashedIssueHandler handlerA = new HashedIssueHandler(context);
+			List<IssueEventGroup> groupsA = TestSingleIssueHandlerGrouped(handlerA, crashA);
+
+			HashedIssueHandler handlerB = new HashedIssueHandler(context);
+			List<IssueEventGroup> groupsB = TestSingleIssueHandlerGrouped(handlerB, crashB);
+
+			Assert.AreEqual(1, groupsA.Count);
+			Assert.AreEqual(1, groupsB.Count);
+			Assert.IsFalse(groupsA[0].Keys.SetEquals(groupsB[0].Keys), "Different crash summaries should produce different hash keys");
+		}
+
+		[TestMethod]
+		public void HashedIssueHandler_NoCrashSummaryFallsBackToFullRender()
+		{
+			// A callstack event WITHOUT CrashSummary should hash on the full render (not match a different callstack)
+			string[] crashA =
+			{
+				"{ \"time\":\"2025-01-01T01:00:00.000Z\",\"level\":\"Error\",\"message\":\"LogWindows: Error: [Callstack] 0x00007ffe5b17b248 UnrealEditor-Core.dll!FDebug::CheckVerifyFailedImpl2() []\",\"id\":301,\"line\":0,\"lineCount\":2,\"properties\":{ \"_channel\":{ \"$type\":\"Channel\",\"$text\":\"LogWindows\"},\"_severity\":{ \"$type\":\"Severity\",\"$text\":\"Error\"},\"Callstack\":true} }",
+				"{ \"time\":\"2025-01-01T01:00:00.001Z\",\"level\":\"Error\",\"message\":\"LogWindows: Error: [Callstack] 0x00007ffe589a608c UnrealEditor-Engine.dll!UWorld::Tick() []\",\"id\":301,\"line\":1,\"lineCount\":2,\"properties\":{ \"_channel\":{ \"$type\":\"Channel\",\"$text\":\"LogWindows\"},\"_severity\":{ \"$type\":\"Severity\",\"$text\":\"Error\"}} }",
+			};
+
+			string[] crashB =
+			{
+				"{ \"time\":\"2025-01-01T01:00:00.000Z\",\"level\":\"Error\",\"message\":\"LogWindows: Error: [Callstack] 0x00007ffe1234abcd UnrealEditor-Slate.dll!SWidget::Paint() []\",\"id\":301,\"line\":0,\"lineCount\":2,\"properties\":{ \"_channel\":{ \"$type\":\"Channel\",\"$text\":\"LogWindows\"},\"_severity\":{ \"$type\":\"Severity\",\"$text\":\"Error\"},\"Callstack\":true} }",
+				"{ \"time\":\"2025-01-01T01:00:00.001Z\",\"level\":\"Error\",\"message\":\"LogWindows: Error: [Callstack] 0x00007ffe5678dcba UnrealEditor-SlateCore.dll!FSlateApplication::Tick() []\",\"id\":301,\"line\":1,\"lineCount\":2,\"properties\":{ \"_channel\":{ \"$type\":\"Channel\",\"$text\":\"LogWindows\"},\"_severity\":{ \"$type\":\"Severity\",\"$text\":\"Error\"}} }",
+			};
+
+			IssueHandlerContext context = new IssueHandlerContext(new StreamId("Release-XX.XX"), new TemplateId(), "Compile ProjectX", new Dictionary<string, string>(), new Dictionary<string, string>());
+
+			HashedIssueHandler handlerA = new HashedIssueHandler(context);
+			List<IssueEventGroup> groupsA = TestSingleIssueHandlerGrouped(handlerA, crashA);
+
+			HashedIssueHandler handlerB = new HashedIssueHandler(context);
+			List<IssueEventGroup> groupsB = TestSingleIssueHandlerGrouped(handlerB, crashB);
+
+			Assert.AreEqual(1, groupsA.Count);
+			Assert.AreEqual(1, groupsB.Count);
+
+			// Without CrashSummary, full render is used — different callstacks produce different hashes
+			Assert.IsFalse(groupsA[0].Keys.SetEquals(groupsB[0].Keys), "Callstack events without CrashSummary should hash on full render");
+		}
+
 		[TestMethod]
 		public void ShaderIssueHandler_GroupsBySourceFileWhenNoMaterialProperty()
 		{
diff --git a/Engine/Source/Programs/Shared/EpicGames.Horde/Issues/Handlers/HashedIssueHandler.cs b/Engine/Source/Programs/Shared/EpicGames.Horde/Issues/Handlers/HashedIssueHandler.cs
--- a/Engine/Source/Programs/Shared/EpicGames.Horde/Issues/Handlers/HashedIssueHandler.cs
+++ b/Engine/Source/Programs/Shared/EpicGames.Horde/Issues/Handlers/HashedIssueHandler.cs
@@ -144,6 +144,15 @@ namespace EpicGames.Horde.Issues.Handlers
 				foreach (IssueEvent stepEvent in item.Value)
 				{
 					string hashSource = stepEvent.Render();
+
+					// For crash callstack events with an assertion/ensure summary, use the
+					// summary for platform-agnostic hashing instead of the full callstack render
+					string? crashSummary = GetStringProperty(stepEvent, "CrashSummary");
+					if (crashSummary != null)
+					{
+						hashSource = crashSummary;
+					}
+
 					string? channel = null;
 					eventScopes.TryGetValue(stepEvent, out channel);
 
@@ -220,6 +229,32 @@ namespace EpicGames.Horde.Issues.Handlers
 			return issues;
 		}
 
+		/// <summary>
+		/// Gets a string property value from the event's JSON properties
+		/// </summary>
+		static string? GetStringProperty(IssueEvent logEvent, string propertyName)
+		{
+			try
+			{
+				foreach (JsonLogEvent line in logEvent.Lines)
+				{
+					using JsonDocument document = JsonDocument.Parse(line.Data);
+					if (document.RootElement.TryGetProperty("properties", out JsonElement properties) &&
+						properties.ValueKind == JsonValueKind.Object &&
+						properties.TryGetProperty(propertyName, out JsonElement value) &&
+						value.ValueKind == JsonValueKind.String)
+					{
+						return value.GetString();
+					}
+				}
+			}
+			catch
+			{
+				// Ignore parsing errors
+			}
+			return null;
+		}
+
 		static void AddFileMetadata(IssueEventGroup issue, IssueEvent stepEvent)
 		{
 			try
diff --git a/Engine/Source/Programs/Shared/EpicGames.LogMatchers.Tests/AutomationToolLogEventParserTests.cs b/Engine/Source/Programs/Shared/EpicGames.LogMatchers.Tests/AutomationToolLogEventParserTests.cs
--- a/Engine/Source/Programs/Shared/EpicGames.LogMatchers.Tests/AutomationToolLogEventParserTests.cs
+++ b/Engine/Source/Programs/Shared/EpicGames.LogMatchers.Tests/AutomationToolLogEventParserTests.cs
@@ -712,6 +712,84 @@ namespace EpicGames.LogMatchers.Tests
 			}
 		}
 
+		[TestMethod]
+		public void LogChannelMatcher_CrashCallstackGrouping()
+		{
+			// Critical error with callstack and trailing blank lines - all should be grouped into one event
+			string[] lines =
+			[
+				@"LogWindows: Error: === Critical error: ===",
+				@"LogWindows: Error:",
+				@"LogWindows: Error: Assertion failed: Pair != nullptr [File:Map.h] [Line: 728]",
+				@"LogWindows: Error:",
+				@"LogWindows: Error: [Callstack] 0x00007ffe5b17b248 UnrealEditor-Core.dll!FDebug::CheckVerifyFailedImpl2() []",
+				@"LogWindows: Error: [Callstack] 0x00007ffe589a608c UnrealEditor-Engine.dll!UWorld::Tick() []",
+				@"LogWindows: Error: [Callstack] 0x00007ffe89f44cb0 KERNEL32.DLL!UnknownFunction []",
+				@"LogWindows: Error:",
+				@"LogWindows: Error:",
+			];
+
+			List<LogEvent> logEvents = Parse(lines);
+			// All 9 lines (header + context + callstack + trailing blank lines) should be grouped
+			CheckEventGroup(logEvents, 0, 9, LogLevel.Error, KnownLogEvents.Engine_LogChannel);
+		}
+
+		[TestMethod]
+		public void LogChannelMatcher_TrailingBlankConsumptionStopsAtRealContent()
+		{
+			// A real error line after trailing blanks must NOT be consumed into the crash group
+			string[] lines =
+			[
+				@"LogWindows: Error: === Critical error: ===",
+				@"LogWindows: Error:",
+				@"LogWindows: Error: Assertion failed: Pair != nullptr [File:Map.h] [Line: 728]",
+				@"LogWindows: Error: [Callstack] 0x00007ffe5b17b248 UnrealEditor-Core.dll!FDebug::CheckVerifyFailedImpl2() []",
+				@"LogWindows: Error: [Callstack] 0x00007ffe89f44cb0 KERNEL32.DLL!UnknownFunction []",
+				@"LogWindows: Error:",
+				@"LogWindows: Error: A completely unrelated error message",
+			];
+
+			List<LogEvent> logEvents = Parse(lines);
+			// First 6 lines grouped (header + context + callstack + trailing blank)
+			CheckEventGroup(logEvents.Slice(0, 6), 0, 6, LogLevel.Error, KnownLogEvents.Engine_LogChannel);
+			// 7th line is a separate event
+			CheckEventGroup(logEvents.Slice(6, 1), 6, 1, LogLevel.Error, KnownLogEvents.Engine_LogChannel);
+		}
+
+		[TestMethod]
+		public void LogChannelMatcher_CallstackOnlyGrouping()
+		{
+			// [Callstack] lines without a preceding crash header should still be grouped, including trailing blanks
+			string[] lines =
+			[
+				@"LogWindows: Error: [Callstack] 0x00007ffe5b17b248 UnrealEditor-Core.dll!FDebug::CheckVerifyFailedImpl2() []",
+				@"LogWindows: Error: [Callstack] 0x00007ffe589a608c UnrealEditor-Engine.dll!UWorld::Tick() []",
+				@"LogWindows: Error: [Callstack] 0x00007ffe89f44cb0 KERNEL32.DLL!UnknownFunction []",
+				@"LogWindows: Error:",
+			];
+
+			List<LogEvent> logEvents = Parse(lines);
+			// All 4 lines (callstack + trailing blank) should be grouped
+			CheckEventGroup(logEvents, 0, 4, LogLevel.Error, KnownLogEvents.Engine_LogChannel);
+		}
+
+		[TestMethod]
+		public void LogChannelMatcher_CrashHeaderWithoutCallstackFallsThrough()
+		{
+			// A crash header without [Callstack] lines within range should fall through to normal behavior
+			string[] lines =
+			[
+				@"LogWindows: Error: === Critical error: ===",
+				@"LogWindows: Error: Some other error info",
+			];
+
+			List<LogEvent> logEvents = Parse(lines);
+			// Without callstack lines, each line is separate (second line is not a hanging indent)
+			Assert.AreEqual(2, logEvents.Count);
+			Assert.AreEqual(KnownLogEvents.Engine_LogChannel, logEvents[0].Id);
+			Assert.AreEqual(KnownLogEvents.Engine_LogChannel, logEvents[1].Id);
+		}
+
 		[TestMethod]
 		public void ShaderEventMatcher()
 		{
diff --git a/Engine/Source/Programs/Shared/EpicGames.LogMatchers/CrashEventMatcher.cs b/Engine/Source/Programs/Shared/EpicGames.LogMatchers/CrashEventMatcher.cs
--- a/Engine/Source/Programs/Shared/EpicGames.LogMatchers/CrashEventMatcher.cs
+++ b/Engine/Source/Programs/Shared/EpicGames.LogMatchers/CrashEventMatcher.cs
@@ -21,14 +21,22 @@ namespace EpicGames.LogMatchers
 		{
 			if (cursor.Contains("begin: stack for UAT"))
 			{
-				for (int maxOffset = 1; maxOffset < 100; maxOffset++)
+				LogEventBuilder builder = new(cursor);
+				for (int count = 0; count < 200; count++)
 				{
-					if (cursor.Contains(maxOffset, "end: stack for UAT"))
+					if (builder.Next.CurrentLine == null)
+					{
+						break;
+					}
+					builder.MoveNext();
+					if (builder.Current.Contains("end: stack for UAT"))
 					{
-						LogEventBuilder builder = new(cursor, lineCount: maxOffset + 1);
 						return builder.ToMatch(LogEventPriority.BelowNormal, GetLogLevel(cursor), KnownLogEvents.Engine_Crash);
 					}
 				}
+				// Found "begin: stack for UAT" but not the closing marker yet; return a partial
+				// match so the parser waits for more data rather than falling through line-by-line.
+				return builder.ToMatch(LogEventPriority.BelowNormal, GetLogLevel(cursor), KnownLogEvents.Engine_Crash, isComplete: false);
 			}
 			if (cursor.IsMatch(s_appErrorPattern))
 			{
diff --git a/Engine/Source/Programs/Shared/EpicGames.LogMatchers/LogChannelEventMatcher.cs b/Engine/Source/Programs/Shared/EpicGames.LogMatchers/LogChannelEventMatcher.cs
--- a/Engine/Source/Programs/Shared/EpicGames.LogMatchers/LogChannelEventMatcher.cs
+++ b/Engine/Source/Programs/Shared/EpicGames.LogMatchers/LogChannelEventMatcher.cs
@@ -18,6 +18,21 @@ namespace EpicGames.LogMatchers
 			@"(?<severity>Error|Warning|Display): "
 		);
 
+		static readonly Regex s_callstackLine = new(@"\[Callstack\]");
+		static readonly Regex s_crashHeader = new(@"=== Critical error: ===|=== Handled ensure: ===");
+
+		// Extracts the assertion/ensure condition and filename for platform-agnostic hashing
+		static readonly Regex s_assertPattern = new(
+			@"(?<type>Assertion failed|Ensure condition failed): (?<condition>.+?)\s+\[File:(?:.*[/\\])?(?<file>[^\]]+)\]");
+
+		// Matches blank log channel lines like "LogWindows: Error:" that UE appends as padding after crash output
+		static readonly Regex s_blankChannelLine = new(
+			@"^(\s*)" +
+			@"(?:\[[\d\.\-: ]+\])*" +
+			@"[a-zA-Z_][a-zA-Z0-9_]*:\s*" +
+			@"(?:Error|Warning|Display):\s*$"
+		);
+
 		/// <inheritdoc/>
 		public LogEventMatch? Match(ILogCursor input)
 		{
@@ -39,6 +54,72 @@ namespace EpicGames.LogMatchers
 					_ => LogLevel.Information,
 				};
 
+				// Crash header with callstack: consume header + context + all [Callstack] frames + trailing blank lines
+				if (s_crashHeader.IsMatch(input.CurrentLine ?? ""))
+				{
+					int callstackOffset = -1;
+					for (int i = 1; i <= 20; i++)
+					{
+						if (!input.TryGetLine(i, out string? line))
+						{
+							break;
+						}
+						if (line.Contains("[Callstack]"))
+						{
+							callstackOffset = i;
+							break;
+						}
+					}
+
+					if (callstackOffset > 0)
+					{
+						builder.AddProperty("Callstack", true);
+
+						// Scan intermediate lines for an assertion/ensure message to use as
+						// a platform-agnostic summary for cross-platform issue matching
+						for (int j = 1; j < callstackOffset; j++)
+						{
+							if (input.TryGetLine(j, out string? intermediateLine))
+							{
+								Match assertMatch = s_assertPattern.Match(intermediateLine);
+								if (assertMatch.Success)
+								{
+									string summary = $"{assertMatch.Groups["type"].Value}: {assertMatch.Groups["condition"].Value.Trim()} [{assertMatch.Groups["file"].Value}]";
+									builder.AddProperty("CrashSummary", summary);
+									break;
+								}
+							}
+						}
+
+						// Skip directly to the first [Callstack] line using the known offset
+						builder.MoveNext(callstackOffset - 1);
+						while (builder.Next.IsMatch(s_callstackLine))
+						{
+							builder.MoveNext();
+						}
+						while (builder.Next.IsMatch(s_blankChannelLine))
+						{
+							builder.MoveNext();
+						}
+						return builder.ToMatch(LogEventPriority.Low, level, KnownLogEvents.Engine_LogChannel);
+					}
+				}
+
+				// [Callstack] line without preceding header: group consecutive frames + trailing blank lines
+				if (s_callstackLine.IsMatch(input.CurrentLine ?? ""))
+				{
+					builder.AddProperty("Callstack", true);
+					while (builder.Next.IsMatch(s_callstackLine))
+					{
+						builder.MoveNext();
+					}
+					while (builder.Next.IsMatch(s_blankChannelLine))
+					{
+						builder.MoveNext();
+					}
+					return builder.ToMatch(LogEventPriority.Low, level, KnownLogEvents.Engine_LogChannel);
+				}
+
 				return builder.ToMatchWithHangingLines(LogEventPriority.Low, level, KnownLogEvents.Engine_LogChannel);
 
 			}
