Horde: Updating schema documentation

#rnx

diff --git a/Engine/Source/Programs/Horde/Docs/Config/Schema/AclActions.md b/Engine/Source/Programs/Horde/Docs/Config/Schema/AclActions.md
--- a/Engine/Source/Programs/Horde/Docs/Config/Schema/AclActions.md
+++ b/Engine/Source/Programs/Horde/Docs/Config/Schema/AclActions.md
@@ -145,17 +145,6 @@
 | `DeleteAccount` |  |
 | `ViewAccount` |  |
 
-## Agents
-
-| Name | Description |
-| ---- | ----------- |
-| `CreateAgent` |  |
-| `CreateWorkstationAgent` |  |
-| `UpdateAgent` |  |
-| `DeleteAgent` |  |
-| `ViewAgent` |  |
-| `ListAgents` |  |
-
 ## Artifacts
 
 | Name | Description |
@@ -180,15 +169,6 @@
 | ---- | ----------- |
 | `AddUpdateDeleteBuildHealthFilter` |  |
 
-## Compute
-
-| Name | Description |
-| ---- | ----------- |
-| `AddComputeTasks` |  |
-| `GetComputeTasks` |  |
-| `UbaCacheRead` |  |
-| `UbaCacheWrite` |  |
-
 ## Devices
 
 | Name | Description |
@@ -207,24 +187,6 @@
 | `RetryJobStep` |  |
 | `ViewJob` |  |
 
-## Leases
-
-| Name | Description |
-| ---- | ----------- |
-| `ViewLeases` |  |
-| `ViewLeaseTasks` |  |
-
-## Logs
-
-| Name | Description |
-| ---- | ----------- |
-| `CreateLog` |  |
-| `UpdateLog` |  |
-| `ViewLog` |  |
-| `WriteLogData` |  |
-| `CreateEvent` |  |
-| `ViewEvent` |  |
-
 ## Notices
 
 | Name | Description |
@@ -239,16 +201,6 @@
 | ---- | ----------- |
 | `CreateSubscription` |  |
 
-## Pools
-
-| Name | Description |
-| ---- | ----------- |
-| `CreatePool` |  |
-| `UpdatePool` |  |
-| `DeletePool` |  |
-| `ViewPool` |  |
-| `ListPools` |  |
-
 ## Projects
 
 | Name | Description |
@@ -274,21 +226,6 @@
 | `DeleteAccount` |  |
 | `ViewAccount` |  |
 
-## Sessions
-
-| Name | Description |
-| ---- | ----------- |
-| `CreateSession` |  |
-| `ViewSession` |  |
-
-## Software
-
-| Name | Description |
-| ---- | ----------- |
-| `UploadSoftware` |  |
-| `DownloadSoftware` |  |
-| `DeleteSoftware` |  |
-
 ## Streams
 
 | Name | Description |
@@ -802,6 +739,19 @@
 | `DownloadSoftware` |  |
 | `DeleteSoftware` |  |
 
+## Storage
+
+| Name | Description |
+| ---- | ----------- |
+| `ReadBlobs` |  |
+| `WriteBlobs` |  |
+| `ReadAliases` |  |
+| `WriteAliases` |  |
+| `DeleteAliases` |  |
+| `ReadRefs` |  |
+| `WriteRefs` |  |
+| `DeleteRefs` |  |
+
 ## Streams
 
 | Name | Description |
@@ -997,6 +947,12 @@
 | `ViewChanges` |  |
 | `ViewTemplate` |  |
 
+## Symbols
+
+| Name | Description |
+| ---- | ----------- |
+| `ReadSymbols` |  |
+
 ## TestData
 
 | Name | Description |
diff --git a/Engine/Source/Programs/Horde/Docs/Config/Schema/Globals.md b/Engine/Source/Programs/Horde/Docs/Config/Schema/Globals.md
--- a/Engine/Source/Programs/Horde/Docs/Config/Schema/Globals.md
+++ b/Engine/Source/Programs/Horde/Docs/Config/Schema/Globals.md
@@ -56,12 +56,12 @@ Name | Description
 Name | Description
 ---- | -----------
 `build` | [BuildConfig](#buildconfig)<br>Configuration for the build plugin
-`compute` | [ComputeConfig](#computeconfig)<br>Configuration for the compute plugin
 `analytics` | [AnalyticsConfig](#analyticsconfig)<br>Configuration for the analytics plugin
+`compute` | [ComputeConfig](#computeconfig)<br>Configuration for the compute plugin
 `exampleplugin` | [ExamplePluginConfig](#examplepluginconfig)<br>Configuration for the exampleplugin plugin
-`experimental` | [ExperimentalConfig](#experimentalconfig)<br>Configuration for the experimental plugin
 `storage` | [StorageConfig](#storageconfig)<br>Configuration for the storage plugin
 `symbols` | [SymbolsConfig](#symbolsconfig)<br>Configuration for the symbols plugin
+`experimental` | [ExperimentalConfig](#experimentalconfig)<br>Configuration for the experimental plugin
 `secrets` | [SecretsConfig](#secretsconfig)<br>Configuration for the secrets plugin
 `tools` | [ToolsConfig](#toolsconfig)<br>Configuration for the tools plugin
 
@@ -215,6 +215,12 @@ Name | Description
 `excludeActions` | `string[]`<br>
 `extends` | `string[]`<br>
 
+## AnalyticsConfig
+
+Name | Description
+---- | -----------
+`stores` | [TelemetryStoreConfig](Telemetry.md)`[]`<br>
+
 ## ComputeConfig
 
 Name | Description
@@ -382,30 +388,24 @@ Name | Description
 `description` | `string`<br>
 `computeId` | `string`<br>
 
-## AnalyticsConfig
-
-Name | Description
----- | -----------
-`stores` | [TelemetryStoreConfig](Telemetry.md)`[]`<br>
-
 ## ExamplePluginConfig
 
 Name | Description
 ---- | -----------
 `greetingMessage` | `string`<br>
 
-## ExperimentalConfig
-
-Name | Description
----- | -----------
-`notifications` | [NotificationConfig](Notification.md)`[]`<br>
-
 ## StorageConfig
 
 Name | Description
 ---- | -----------
 `enableGc` | `boolean`<br>
 `enableGcVerification` | `boolean`<br>
+`gcWorkerCount` | `integer`<br>
+`gcQueueLimit` | `integer`<br>
+`enableReliableGcEnqueue` | `boolean`<br>
+`enablePerNamespaceQueueIsolation` | `boolean`<br>
+`gcBlobIngestionBatchSize` | `integer`<br>
+`gcBatchSize` | `integer`<br>
 `backends` | [BackendConfig](#backendconfig)`[]`<br>
 `namespaces` | [NamespaceConfig](#namespaceconfig)`[]`<br>
 
@@ -477,6 +477,12 @@ Name | Description
 `public` | `boolean`<br>
 `acl` | [AclConfig](#aclconfig)<br>
 
+## ExperimentalConfig
+
+Name | Description
+---- | -----------
+`notifications` | [NotificationConfig](Notification.md)`[]`<br>
+
 ## SecretsConfig
 
 Name | Description
diff --git a/Engine/Source/Programs/Horde/Docs/Config/Schema/Projects.md b/Engine/Source/Programs/Horde/Docs/Config/Schema/Projects.md
--- a/Engine/Source/Programs/Horde/Docs/Config/Schema/Projects.md
+++ b/Engine/Source/Programs/Horde/Docs/Config/Schema/Projects.md
@@ -15,6 +15,7 @@ Name | Description
 `pools` | [PoolConfig](#poolconfig)`[]`<br>
 `categories` | [ProjectCategoryConfig](#projectcategoryconfig)`[]`<br>
 `jobOptions` | [JobOptions](#joboptions)<br>
+`retentionOptions` | [RetentionOptions](#retentionoptions)<br>
 `workspaceTypes` | `string` `->` [WorkspaceConfig](#workspaceconfig)<br>
 `ignorePatternFiles` | `string[]`<br>
 `telemetryStoreId` | `string`<br>
@@ -168,6 +169,14 @@ Name | Description
 `containerEngineExecutable` | `string`<br>Container engine executable (docker or with full path like /usr/bin/podman)
 `extraArguments` | `string`<br>Additional arguments to pass to container engine
 
+## RetentionOptions
+
+Name | Description
+---- | -----------
+`enableJobExpiration` | `boolean`<br>
+`enableArtifactExpiration` | `boolean`<br>
+`maxJobDeletionsPerTick` | `integer`<br>
+
 ## WorkspaceConfig
 
 Name | Description
diff --git a/Engine/Source/Programs/Horde/Docs/Config/Schema/Streams.md b/Engine/Source/Programs/Horde/Docs/Config/Schema/Streams.md
--- a/Engine/Source/Programs/Horde/Docs/Config/Schema/Streams.md
+++ b/Engine/Source/Programs/Horde/Docs/Config/Schema/Streams.md
@@ -21,6 +21,7 @@ Name | Description
 `notificationChannelFilter` | `string`<br>
 `triageChannel` | `string`<br>
 `jobOptions` | [JobOptions](#joboptions)<br>
+`retentionOptions` | [RetentionOptions](#retentionoptions)<br>
 `telemetryStoreId` | `string`<br>
 `autoSdkView` | `string[]`<br>
 `defaultPreflightTemplate` | `string`<br>
@@ -79,6 +80,14 @@ Name | Description
 `containerEngineExecutable` | `string`<br>Container engine executable (docker or with full path like /usr/bin/podman)
 `extraArguments` | `string`<br>Additional arguments to pass to container engine
 
+## RetentionOptions
+
+Name | Description
+---- | -----------
+`enableJobExpiration` | `boolean`<br>
+`enableArtifactExpiration` | `boolean`<br>
+`maxJobDeletionsPerTick` | `integer`<br>
+
 ## DefaultPreflightConfig
 
 Name | Description
