Horde: Update release notes

#rnx

diff --git a/Engine/Source/Programs/Horde/Docs/ReleaseNotes.md b/Engine/Source/Programs/Horde/Docs/ReleaseNotes.md
--- a/Engine/Source/Programs/Horde/Docs/ReleaseNotes.md
+++ b/Engine/Source/Programs/Horde/Docs/ReleaseNotes.md
@@ -1,5 +1,26 @@
 # Release Notes
 
+## 2026-03-05
+
+* EpicGames.LogMatchers: Fix a static analysis warning by specifying ordinal string comparison when checking if a log line contains the callstack indicator. (51456221)
+* Group crash callstack lines into single events in log matchers (51451038)
+* Add Storage Dashboard frontend view (51433116)
+* Add Storage Dashboard backend: GC metrics, collection stats, API endpoint (51423509)
+* Add per-stream RetentionOptions with safety valves for job and artifact expiration (51411407)
+* Fix crash in GetIssueEventsAsync when logIds contains invalid values (51397115)
+* Add jobId, batchId, stepId, workflowId, fingerprintType params to debug create-test-issue endpoint (51393039)
+* Add fingerprintType parameter to debug create-test-issue endpoint (51390872)
+* Add debug endpoint for creating test issues with suspects (51379057)
+* Suppress Slack @-mentions for low-context issue types (Hashed, Gauntlet) (51373817)
+* Fix getMore failures in AgentCollection.FindAsync (51340180)
+* Detect edge server change in PerforceMaterializer to force clean re-sync (51338317)
+* EpicGames.Perforce: Fix unshelve not invalidating deleted files in workspace tracking (51298759)
+* Fix empty pool notification spam, Only fire OnJobScheduled once per batch instead of every 5s queue refresh, preventing 73K+ entries accumulating in Redis batch queue (51292704)
+* Fix OnJobScheduled handler leak in NotificationService.DisposeAsync (51292392)
+* Improve truncation of over-sized log messages (51292253)
+* Update Build Health documentation with comprehensive workflow reference and examples (51273164)
+* Config loading fault tolerance: missing base references and diamond include scoping (51269740)
+
 ## 2026-02-26
 
 * Make agent session conflict detection disabled as default (51262410)
