Horde: Only register TunnelService when tunnel port is configured

Also make StopAsync log exceptions as warnings.

diff --git a/Engine/Source/Programs/Horde/Plugins/Compute/HordeServer.Compute/Compute/TunnelService.cs b/Engine/Source/Programs/Horde/Plugins/Compute/HordeServer.Compute/Compute/TunnelService.cs
--- a/Engine/Source/Programs/Horde/Plugins/Compute/HordeServer.Compute/Compute/TunnelService.cs
+++ b/Engine/Source/Programs/Horde/Plugins/Compute/HordeServer.Compute/Compute/TunnelService.cs
@@ -2,7 +2,6 @@
 
 using System.Net;
 using System.Net.Sockets;
-using EpicGames.Core;
 using EpicGames.Horde.Compute.Clients;
 using Microsoft.Extensions.Hosting;
 using Microsoft.Extensions.Logging;
@@ -65,8 +64,21 @@ namespace HordeServer.Compute
 		{
 			await _cancellationSource.CancelAsync();
 
-			await _serverTask!.IgnoreCanceledExceptionsAsync().ConfigureAwait(false);
-			_serverTask = null;
+			if (_serverTask != null)
+			{
+				try
+				{
+					await _serverTask.ConfigureAwait(false);
+				}
+				catch (OperationCanceledException)
+				{
+				}
+				catch (Exception ex)
+				{
+					_logger.LogWarning(ex, "Tunnel service task completed with an error");
+				}
+				_serverTask = null;
+			}
 		}
 
 		async Task HandleClientAsync(TcpClient client, CancellationToken cancellationToken)
diff --git a/Engine/Source/Programs/Horde/Plugins/Compute/HordeServer.Compute/ComputePlugin.cs b/Engine/Source/Programs/Horde/Plugins/Compute/HordeServer.Compute/ComputePlugin.cs
--- a/Engine/Source/Programs/Horde/Plugins/Compute/HordeServer.Compute/ComputePlugin.cs
+++ b/Engine/Source/Programs/Horde/Plugins/Compute/HordeServer.Compute/ComputePlugin.cs
@@ -96,9 +96,11 @@ namespace HordeServer
 			// Though most fleet managers are created on a per-pool basis
 			services.AddSingleton<IFleetManager>(ctx => ctx.GetRequiredService<IFleetManagerFactory>().CreateFleetManager(FleetManagerType.Default));
 
-			// Run the tunnel service for all run modes
-			services.AddSingleton<TunnelService>();
-			services.AddHostedService<TunnelService>(sp => sp.GetRequiredService<TunnelService>());
+			if (_staticComputeConfig.ComputeTunnelPort != 0)
+			{
+				services.AddSingleton<TunnelService>();
+				services.AddHostedService<TunnelService>(sp => sp.GetRequiredService<TunnelService>());
+			}
 
 			// Runs the agent relay service for all run modes to notify long-polling requests
 			services.AddSingleton<AgentRelayService>();
