Fix build warnings in Storage and Build plugins

Add missing XML doc comments to CollectionStatsService and GcMetricsCollection.
Remove unused using directive in CollectionStatsService.
Make BuildStatusResponse static in StorageController (CA1822).
Suppress CA1505 on CreateOrUpdateWorkflowThreadAsync in SlackNotificationSink.

diff --git a/Engine/Source/Programs/Horde/Plugins/Build/HordeServer.Build/Notifications/Sinks/SlackNotificationSink.cs b/Engine/Source/Programs/Horde/Plugins/Build/HordeServer.Build/Notifications/Sinks/SlackNotificationSink.cs
--- a/Engine/Source/Programs/Horde/Plugins/Build/HordeServer.Build/Notifications/Sinks/SlackNotificationSink.cs
+++ b/Engine/Source/Programs/Horde/Plugins/Build/HordeServer.Build/Notifications/Sinks/SlackNotificationSink.cs
@@ -1057,7 +1057,9 @@ namespace HordeServer.Notifications.Sinks
 			NoSuspects = 32,
 		}
 
+#pragma warning disable CA1505
 		async Task CreateOrUpdateWorkflowThreadAsync(string triageChannel, IIssue issue, IIssueSpan span, IReadOnlyList<IIssueSpan> spans, WorkflowConfig workflow, CancellationToken cancellationToken)
+#pragma warning restore CA1505
 		{
 			Uri issueUrl = GetIssueUrl(issue, span.FirstFailure);
 
diff --git a/Engine/Source/Programs/Horde/Plugins/Storage/HordeServer.Storage/Storage/CollectionStatsService.cs b/Engine/Source/Programs/Horde/Plugins/Storage/HordeServer.Storage/Storage/CollectionStatsService.cs
--- a/Engine/Source/Programs/Horde/Plugins/Storage/HordeServer.Storage/Storage/CollectionStatsService.cs
+++ b/Engine/Source/Programs/Horde/Plugins/Storage/HordeServer.Storage/Storage/CollectionStatsService.cs
@@ -2,7 +2,6 @@
 
 using EpicGames.Core;
 using HordeServer.Server;
-using HordeServer.Utilities;
 using Microsoft.Extensions.Hosting;
 using Microsoft.Extensions.Logging;
 using MongoDB.Bson;
@@ -16,24 +15,45 @@ namespace HordeServer.Storage
 	/// </summary>
 	public class CollectionStatEntry
 	{
+		/// <summary>
+		/// Name of the MongoDB collection
+		/// </summary>
 		[BsonElement("name")]
 		public string CollectionName { get; set; } = string.Empty;
 
+		/// <summary>
+		/// Number of documents in the collection
+		/// </summary>
 		[BsonElement("cnt")]
 		public long DocumentCount { get; set; }
 
+		/// <summary>
+		/// Total data size in bytes
+		/// </summary>
 		[BsonElement("dsz")]
 		public long DataSizeBytes { get; set; }
 
+		/// <summary>
+		/// Total index size in bytes
+		/// </summary>
 		[BsonElement("isz")]
 		public long IndexSizeBytes { get; set; }
 
+		/// <summary>
+		/// Total storage size in bytes
+		/// </summary>
 		[BsonElement("ssz")]
 		public long StorageSizeBytes { get; set; }
 
+		/// <summary>
+		/// Average document size in bytes
+		/// </summary>
 		[BsonElement("asz")]
 		public long AvgDocSizeBytes { get; set; }
 
+		/// <summary>
+		/// Ratio of storage size to data size indicating fragmentation
+		/// </summary>
 		[BsonElement("frag")]
 		public double FragmentationRatio { get; set; }
 	}
@@ -43,11 +63,20 @@ namespace HordeServer.Storage
 	/// </summary>
 	public class CollectionStatsDocument
 	{
+		/// <summary>
+		/// Unique identifier for this document
+		/// </summary>
 		public ObjectId Id { get; set; }
 
+		/// <summary>
+		/// Time this snapshot was captured
+		/// </summary>
 		[BsonElement("tm")]
 		public DateTime TimeUtc { get; set; }
 
+		/// <summary>
+		/// Stats for each monitored collection
+		/// </summary>
 		[BsonElement("cols")]
 		public List<CollectionStatEntry> Collections { get; set; } = new();
 	}
@@ -81,6 +110,9 @@ namespace HordeServer.Storage
 		readonly ITicker _cleanupTicker;
 #pragma warning restore CA2213
 
+		/// <summary>
+		/// Constructor
+		/// </summary>
 		public CollectionStatsService(IMongoService mongoService, IClock clock, ILogger<CollectionStatsService> logger)
 		{
 			_mongoService = mongoService;
diff --git a/Engine/Source/Programs/Horde/Plugins/Storage/HordeServer.Storage/Storage/GcMetricsCollection.cs b/Engine/Source/Programs/Horde/Plugins/Storage/HordeServer.Storage/Storage/GcMetricsCollection.cs
--- a/Engine/Source/Programs/Horde/Plugins/Storage/HordeServer.Storage/Storage/GcMetricsCollection.cs
+++ b/Engine/Source/Programs/Horde/Plugins/Storage/HordeServer.Storage/Storage/GcMetricsCollection.cs
@@ -18,16 +18,59 @@ namespace HordeServer.Storage
 	/// </summary>
 	public record GcMetricEntry
 	{
+		/// <summary>
+		/// Namespace for this GC metric
+		/// </summary>
 		public NamespaceId NamespaceId { get; init; }
+
+		/// <summary>
+		/// Number of blobs deleted during this sweep
+		/// </summary>
 		public long BlobsDeleted { get; init; }
+
+		/// <summary>
+		/// Number of blobs checked during this sweep
+		/// </summary>
 		public long BlobsChecked { get; init; }
+
+		/// <summary>
+		/// Total bytes freed during this sweep
+		/// </summary>
 		public long BytesFreed { get; init; }
+
+		/// <summary>
+		/// Duration of the sweep in milliseconds
+		/// </summary>
 		public double SweepDurationMs { get; init; }
+
+		/// <summary>
+		/// Whether GC was paused due to queue pressure
+		/// </summary>
 		public bool WasPaused { get; init; }
+
+		/// <summary>
+		/// Current depth of the GC queue
+		/// </summary>
 		public long QueueDepth { get; init; }
+
+		/// <summary>
+		/// Number of blobs ingested since last tick
+		/// </summary>
 		public long BlobsIngested { get; init; }
+
+		/// <summary>
+		/// Number of enqueue failures since last tick
+		/// </summary>
 		public long EnqueueFailures { get; init; }
+
+		/// <summary>
+		/// Number of throttle events since last tick
+		/// </summary>
 		public long ThrottleEvents { get; init; }
+
+		/// <summary>
+		/// Number of refs that expired since last tick
+		/// </summary>
 		public long RefsExpired { get; init; }
 	}
 
@@ -36,41 +79,80 @@ namespace HordeServer.Storage
 	/// </summary>
 	public class GcMetricDocument
 	{
+		/// <summary>
+		/// Unique identifier for this document
+		/// </summary>
 		public ObjectId Id { get; set; }
 
+		/// <summary>
+		/// Time this metric was recorded
+		/// </summary>
 		[BsonElement("tm")]
 		public DateTime TimeUtc { get; set; }
 
+		/// <summary>
+		/// Namespace for this GC metric
+		/// </summary>
 		[BsonElement("ns")]
 		public NamespaceId NamespaceId { get; set; }
 
+		/// <summary>
+		/// Number of blobs deleted
+		/// </summary>
 		[BsonElement("del")]
 		public long BlobsDeleted { get; set; }
 
+		/// <summary>
+		/// Number of blobs checked
+		/// </summary>
 		[BsonElement("chk")]
 		public long BlobsChecked { get; set; }
 
+		/// <summary>
+		/// Total bytes freed
+		/// </summary>
 		[BsonElement("freed")]
 		public long BytesFreed { get; set; }
 
+		/// <summary>
+		/// Duration of the sweep in milliseconds
+		/// </summary>
 		[BsonElement("sweep_ms")]
 		public double SweepDurationMs { get; set; }
 
+		/// <summary>
+		/// Whether GC was paused due to queue pressure
+		/// </summary>
 		[BsonElement("paused")]
 		public bool WasPaused { get; set; }
 
+		/// <summary>
+		/// Current depth of the GC queue
+		/// </summary>
 		[BsonElement("qdepth")]
 		public long QueueDepth { get; set; }
 
+		/// <summary>
+		/// Number of blobs ingested
+		/// </summary>
 		[BsonElement("ing")]
 		public long BlobsIngested { get; set; }
 
+		/// <summary>
+		/// Number of enqueue failures
+		/// </summary>
 		[BsonElement("enq_fail")]
 		public long EnqueueFailures { get; set; }
 
+		/// <summary>
+		/// Number of throttle events
+		/// </summary>
 		[BsonElement("throttle")]
 		public long ThrottleEvents { get; set; }
 
+		/// <summary>
+		/// Number of refs that expired
+		/// </summary>
 		[BsonElement("refs_exp")]
 		public long RefsExpired { get; set; }
 	}
@@ -95,6 +177,9 @@ namespace HordeServer.Storage
 		readonly IClock _clock;
 		readonly ILogger _logger;
 
+		/// <summary>
+		/// Constructor
+		/// </summary>
 		public GcMetricsCollection(StorageService storageService, IMongoService mongoService, IOptionsMonitor<StorageConfig> storageConfig, IClock clock, ILogger<GcMetricsCollection> logger)
 		{
 			_storageService = storageService;
diff --git a/Engine/Source/Programs/Horde/Plugins/Storage/HordeServer.Storage/Storage/StorageController.cs b/Engine/Source/Programs/Horde/Plugins/Storage/HordeServer.Storage/Storage/StorageController.cs
--- a/Engine/Source/Programs/Horde/Plugins/Storage/HordeServer.Storage/Storage/StorageController.cs
+++ b/Engine/Source/Programs/Horde/Plugins/Storage/HordeServer.Storage/Storage/StorageController.cs
@@ -715,7 +715,7 @@ namespace HordeServer.Storage
 			};
 		}
 
-		StorageDashboardStatus BuildStatusResponse(StorageConfig config, StorageServerConfig serverConfig, (List<(NamespaceId Id, DateTime LastTime)> Namespaces, Dictionary<NamespaceId, long> QueueDepths) gcStatus)
+		static StorageDashboardStatus BuildStatusResponse(StorageConfig config, StorageServerConfig serverConfig, (List<(NamespaceId Id, DateTime LastTime)> Namespaces, Dictionary<NamespaceId, long> QueueDepths) gcStatus)
 		{
 			Dictionary<NamespaceId, DateTime> lastGcTimes = gcStatus.Namespaces.ToDictionary(x => x.Id, x => x.LastTime);
 
