Horde: Increase MongoDB batch size for agent reads to reduce getMore errors

Set a batch size of 10000 when fetching agent documents from MongoDB to minimize getMore commands on secondary reads. The agents collection is updated frequently, making multi-batch cursor reads prone to errors when reading from replicas.

diff --git a/Engine/Source/Programs/Horde/Plugins/Compute/HordeServer.Compute/Agents/AgentCollection.cs b/Engine/Source/Programs/Horde/Plugins/Compute/HordeServer.Compute/Agents/AgentCollection.cs
--- a/Engine/Source/Programs/Horde/Plugins/Compute/HordeServer.Compute/Agents/AgentCollection.cs
+++ b/Engine/Source/Programs/Horde/Plugins/Compute/HordeServer.Compute/Agents/AgentCollection.cs
@@ -711,9 +711,10 @@ namespace HordeServer.Agents
 
 			IMongoCollection<AgentDocument> collection = consistentRead ? _agentCollection : _agentCollection.WithReadPreference(ReadPreference.SecondaryPreferred);
 
-			// Materialize all documents up front to avoid long-lived cursors vulnerable to getMore failures on secondary reads.
-			// Post-load processing (document migration, Redis session lookups) can take significant time between cursor batches.
-			List<AgentDocument> allDocuments = await collection.Find(filter).ToListAsync(cancellationToken);
+			// Use a large batch size to minimize getMore commands when reading from secondaries.
+			// The agents collection is updated frequently, making multi-batch cursor reads error-prone.
+			FindOptions findOptions = new() { BatchSize = 10000 };
+			List<AgentDocument> allDocuments = await collection.Find(filter, findOptions).ToListAsync(cancellationToken);
 			allDocuments = await PostLoadAsync(allDocuments, cancellationToken);
 			foreach (IAgent agent in await CreateAgentObjectsBatchedAsync(allDocuments, cancellationToken))
 			{
