Horde: Updating schema documentation

#rnx

diff --git a/Engine/Source/Programs/Horde/Docs/Config/Schema/AclActions.md b/Engine/Source/Programs/Horde/Docs/Config/Schema/AclActions.md
--- a/Engine/Source/Programs/Horde/Docs/Config/Schema/AclActions.md
+++ b/Engine/Source/Programs/Horde/Docs/Config/Schema/AclActions.md
@@ -37,32 +37,6 @@
 | `DeleteAccount` |  |
 | `ViewAccount` |  |
 
-## Notices
-
-| Name | Description |
-| ---- | ----------- |
-| `CreateNotice` |  |
-| `UpdateNotice` |  |
-| `DeleteNotice` |  |
-
-## ServiceAccounts
-
-| Name | Description |
-| ---- | ----------- |
-| `CreateAccount` |  |
-| `UpdateAccount` |  |
-| `DeleteAccount` |  |
-| `ViewAccount` |  |
-
-## Accounts
-
-| Name | Description |
-| ---- | ----------- |
-| `CreateAccount` |  |
-| `UpdateAccount` |  |
-| `DeleteAccount` |  |
-| `ViewAccount` |  |
-
 ## Artifacts
 
 | Name | Description |
@@ -279,6 +253,17 @@
 | `DeleteAccount` |  |
 | `ViewAccount` |  |
 
+## Agents
+
+| Name | Description |
+| ---- | ----------- |
+| `CreateAgent` | Ability to create an agent. This may be done explicitly, or granted to agents to allow them to self-register. |
+| `CreateWorkstationAgent` | Ability to create a workstation agent. Granted to users with hardware to offer. |
+| `UpdateAgent` | Update an agent's name, pools, etc... |
+| `DeleteAgent` | Soft-delete an agent |
+| `ViewAgent` | View an agent |
+| `ListAgents` | List the available agents |
+
 ## Artifacts
 
 | Name | Description |
@@ -303,6 +288,15 @@
 | ---- | ----------- |
 | `AddUpdateDeleteBuildHealthFilter` |  |
 
+## Compute
+
+| Name | Description |
+| ---- | ----------- |
+| `AddComputeTasks` | User can add tasks to the compute cluster |
+| `GetComputeTasks` | User can get and list tasks from the compute cluster |
+| `UbaCacheRead` | User can read from Unreal Build Accelerator cache |
+| `UbaCacheWrite` | User can write to Unreal Build Accelerator cache |
+
 ## Devices
 
 | Name | Description |
@@ -321,6 +315,24 @@
 | `RetryJobStep` |  |
 | `ViewJob` |  |
 
+## Leases
+
+| Name | Description |
+| ---- | ----------- |
+| `ViewLeases` | View all the leases that an agent has worked on |
+| `ViewLeaseTasks` | View the task data for a lease |
+
+## Logs
+
+| Name | Description |
+| ---- | ----------- |
+| `CreateLog` | Ability to create a log. Implicitly granted to agents. |
+| `UpdateLog` | Ability to update log metadata |
+| `ViewLog` | Ability to view a log contents |
+| `WriteLogData` | Ability to write log data |
+| `CreateEvent` | Ability to create events |
+| `ViewEvent` | Ability to view events |
+
 ## Notices
 
 | Name | Description |
@@ -335,6 +347,16 @@
 | ---- | ----------- |
 | `CreateSubscription` |  |
 
+## Pools
+
+| Name | Description |
+| ---- | ----------- |
+| `CreatePool` | Create a global pool of agents |
+| `UpdatePool` | Modify an agent pool |
+| `DeletePool` | Delete an agent pool |
+| `ViewPool` | Ability to view a pool |
+| `ListPools` | View all the available agent pools |
+
 ## Projects
 
 | Name | Description |
@@ -360,6 +382,21 @@
 | `DeleteAccount` |  |
 | `ViewAccount` |  |
 
+## Sessions
+
+| Name | Description |
+| ---- | ----------- |
+| `CreateSession` | Granted to agents to call CreateSession, which returns a bearer token identifying themselves valid to call UpdateSesssion via gRPC. |
+| `ViewSession` | Allows viewing information about an agent session |
+
+## Software
+
+| Name | Description |
+| ---- | ----------- |
+| `UploadSoftware` | Ability to upload new versions of the agent software |
+| `DownloadSoftware` | Ability to download the agent software |
+| `DeleteSoftware` | Ability to delete agent software |
+
 ## Streams
 
 | Name | Description |
@@ -387,6 +424,17 @@
 | `DeleteAccount` |  |
 | `ViewAccount` |  |
 
+## Agents
+
+| Name | Description |
+| ---- | ----------- |
+| `CreateAgent` |  |
+| `CreateWorkstationAgent` |  |
+| `UpdateAgent` |  |
+| `DeleteAgent` |  |
+| `ViewAgent` |  |
+| `ListAgents` |  |
+
 ## Artifacts
 
 | Name | Description |
@@ -411,6 +459,15 @@
 | ---- | ----------- |
 | `AddUpdateDeleteBuildHealthFilter` |  |
 
+## Compute
+
+| Name | Description |
+| ---- | ----------- |
+| `AddComputeTasks` |  |
+| `GetComputeTasks` |  |
+| `UbaCacheRead` |  |
+| `UbaCacheWrite` |  |
+
 ## Devices
 
 | Name | Description |
@@ -429,6 +486,24 @@
 | `RetryJobStep` |  |
 | `ViewJob` |  |
 
+## Leases
+
+| Name | Description |
+| ---- | ----------- |
+| `ViewLeases` |  |
+| `ViewLeaseTasks` |  |
+
+## Logs
+
+| Name | Description |
+| ---- | ----------- |
+| `CreateLog` |  |
+| `UpdateLog` |  |
+| `ViewLog` |  |
+| `WriteLogData` |  |
+| `CreateEvent` |  |
+| `ViewEvent` |  |
+
 ## Notices
 
 | Name | Description |
@@ -443,6 +518,16 @@
 | ---- | ----------- |
 | `CreateSubscription` |  |
 
+## Pools
+
+| Name | Description |
+| ---- | ----------- |
+| `CreatePool` |  |
+| `UpdatePool` |  |
+| `DeletePool` |  |
+| `ViewPool` |  |
+| `ListPools` |  |
+
 ## Projects
 
 | Name | Description |
@@ -474,6 +559,21 @@
 | `DeleteAccount` |  |
 | `ViewAccount` |  |
 
+## Sessions
+
+| Name | Description |
+| ---- | ----------- |
+| `CreateSession` |  |
+| `ViewSession` |  |
+
+## Software
+
+| Name | Description |
+| ---- | ----------- |
+| `UploadSoftware` |  |
+| `DownloadSoftware` |  |
+| `DeleteSoftware` |  |
+
 ## Streams
 
 | Name | Description |
@@ -505,12 +605,12 @@
 
 | Name | Description |
 | ---- | ----------- |
-| `CreateAgent` | Ability to create an agent. This may be done explicitly, or granted to agents to allow them to self-register. |
-| `CreateWorkstationAgent` | Ability to create a workstation agent. Granted to users with hardware to offer. |
-| `UpdateAgent` | Update an agent's name, pools, etc... |
-| `DeleteAgent` | Soft-delete an agent |
-| `ViewAgent` | View an agent |
-| `ListAgents` | List the available agents |
+| `CreateAgent` |  |
+| `CreateWorkstationAgent` |  |
+| `UpdateAgent` |  |
+| `DeleteAgent` |  |
+| `ViewAgent` |  |
+| `ListAgents` |  |
 
 ## Artifacts
 
@@ -540,10 +640,10 @@
 
 | Name | Description |
 | ---- | ----------- |
-| `AddComputeTasks` | User can add tasks to the compute cluster |
-| `GetComputeTasks` | User can get and list tasks from the compute cluster |
-| `UbaCacheRead` | User can read from Unreal Build Accelerator cache |
-| `UbaCacheWrite` | User can write to Unreal Build Accelerator cache |
+| `AddComputeTasks` |  |
+| `GetComputeTasks` |  |
+| `UbaCacheRead` |  |
+| `UbaCacheWrite` |  |
 
 ## Devices
 
@@ -567,19 +667,19 @@
 
 | Name | Description |
 | ---- | ----------- |
-| `ViewLeases` | View all the leases that an agent has worked on |
-| `ViewLeaseTasks` | View the task data for a lease |
+| `ViewLeases` |  |
+| `ViewLeaseTasks` |  |
 
 ## Logs
 
 | Name | Description |
 | ---- | ----------- |
-| `CreateLog` | Ability to create a log. Implicitly granted to agents. |
-| `UpdateLog` | Ability to update log metadata |
-| `ViewLog` | Ability to view a log contents |
-| `WriteLogData` | Ability to write log data |
-| `CreateEvent` | Ability to create events |
-| `ViewEvent` | Ability to view events |
+| `CreateLog` |  |
+| `UpdateLog` |  |
+| `ViewLog` |  |
+| `WriteLogData` |  |
+| `CreateEvent` |  |
+| `ViewEvent` |  |
 
 ## Notices
 
@@ -599,11 +699,11 @@
 
 | Name | Description |
 | ---- | ----------- |
-| `CreatePool` | Create a global pool of agents |
-| `UpdatePool` | Modify an agent pool |
-| `DeletePool` | Delete an agent pool |
-| `ViewPool` | Ability to view a pool |
-| `ListPools` | View all the available agent pools |
+| `CreatePool` |  |
+| `UpdatePool` |  |
+| `DeletePool` |  |
+| `ViewPool` |  |
+| `ListPools` |  |
 
 ## Projects
 
@@ -640,16 +740,16 @@
 
 | Name | Description |
 | ---- | ----------- |
-| `CreateSession` | Granted to agents to call CreateSession, which returns a bearer token identifying themselves valid to call UpdateSesssion via gRPC. |
-| `ViewSession` | Allows viewing information about an agent session |
+| `CreateSession` |  |
+| `ViewSession` |  |
 
 ## Software
 
 | Name | Description |
 | ---- | ----------- |
-| `UploadSoftware` | Ability to upload new versions of the agent software |
-| `DownloadSoftware` | Ability to download the agent software |
-| `DeleteSoftware` | Ability to delete agent software |
+| `UploadSoftware` |  |
+| `DownloadSoftware` |  |
+| `DeleteSoftware` |  |
 
 ## Streams
 
@@ -846,12 +946,195 @@
 | `TestHealthRead` |  |
 | `TestHealthWrite` |  |
 
-## Tools
+## Accounts
 
 | Name | Description |
 | ---- | ----------- |
-| `DownloadTool` | Ability to download a tool |
-| `UploadTool` | Ability to upload new tool versions |
+| `CreateAccount` |  |
+| `UpdateAccount` |  |
+| `DeleteAccount` |  |
+| `ViewAccount` |  |
+
+## Agents
+
+| Name | Description |
+| ---- | ----------- |
+| `CreateAgent` |  |
+| `CreateWorkstationAgent` |  |
+| `UpdateAgent` |  |
+| `DeleteAgent` |  |
+| `ViewAgent` |  |
+| `ListAgents` |  |
+
+## Artifacts
+
+| Name | Description |
+| ---- | ----------- |
+| `ReadArtifact` |  |
+| `WriteArtifact` |  |
+| `DeleteArtifact` |  |
+| `UploadArtifact` |  |
+| `DownloadArtifact` |  |
+
+## Bisect
+
+| Name | Description |
+| ---- | ----------- |
+| `CreateBisectTask` |  |
+| `UpdateBisectTask` |  |
+| `ViewBisectTask` |  |
+
+## BuildHealth
+
+| Name | Description |
+| ---- | ----------- |
+| `AddUpdateDeleteBuildHealthFilter` |  |
+
+## Compute
+
+| Name | Description |
+| ---- | ----------- |
+| `AddComputeTasks` |  |
+| `GetComputeTasks` |  |
+| `UbaCacheRead` |  |
+| `UbaCacheWrite` |  |
+
+## Devices
+
+| Name | Description |
+| ---- | ----------- |
+| `DeviceRead` |  |
+| `DeviceWrite` |  |
+
+## Jobs
+
+| Name | Description |
+| ---- | ----------- |
+| `CreateJob` |  |
+| `UpdateJob` |  |
+| `DeleteJob` |  |
+| `ExecuteJob` |  |
+| `RetryJobStep` |  |
+| `ViewJob` |  |
+
+## Leases
+
+| Name | Description |
+| ---- | ----------- |
+| `ViewLeases` |  |
+| `ViewLeaseTasks` |  |
+
+## Logs
+
+| Name | Description |
+| ---- | ----------- |
+| `CreateLog` |  |
+| `UpdateLog` |  |
+| `ViewLog` |  |
+| `WriteLogData` |  |
+| `CreateEvent` |  |
+| `ViewEvent` |  |
+
+## Notices
+
+| Name | Description |
+| ---- | ----------- |
+| `CreateNotice` |  |
+| `UpdateNotice` |  |
+| `DeleteNotice` |  |
+
+## Notifications
+
+| Name | Description |
+| ---- | ----------- |
+| `CreateSubscription` |  |
+
+## Pools
+
+| Name | Description |
+| ---- | ----------- |
+| `CreatePool` |  |
+| `UpdatePool` |  |
+| `DeletePool` |  |
+| `ViewPool` |  |
+| `ListPools` |  |
+
+## Projects
+
+| Name | Description |
+| ---- | ----------- |
+| `CreateProject` |  |
+| `DeleteProject` |  |
+| `UpdateProject` |  |
+| `ViewProject` |  |
+
+## Replicators
+
+| Name | Description |
+| ---- | ----------- |
+| `UpdateReplicator` |  |
+| `ViewReplicator` |  |
+
+## Secrets
+
+| Name | Description |
+| ---- | ----------- |
+| `ViewSecret` |  |
+
+## ServiceAccounts
+
+| Name | Description |
+| ---- | ----------- |
+| `CreateAccount` |  |
+| `UpdateAccount` |  |
+| `DeleteAccount` |  |
+| `ViewAccount` |  |
+
+## Sessions
+
+| Name | Description |
+| ---- | ----------- |
+| `CreateSession` |  |
+| `ViewSession` |  |
+
+## Software
+
+| Name | Description |
+| ---- | ----------- |
+| `UploadSoftware` |  |
+| `DownloadSoftware` |  |
+| `DeleteSoftware` |  |
+
+## Storage
+
+| Name | Description |
+| ---- | ----------- |
+| `ReadBlobs` | Ability to read blobs from the storage service |
+| `WriteBlobs` | Ability to write blobs to the storage service |
+| `ReadAliases` | Ability to read aliases from the storage service |
+| `WriteAliases` | Ability to write aliases to the storage service |
+| `DeleteAliases` | Ability to write aliases to the storage service |
+| `ReadRefs` | Ability to read refs from the storage service |
+| `WriteRefs` | Ability to write refs to the storage service |
+| `DeleteRefs` | Ability to delete refs |
+
+## Streams
+
+| Name | Description |
+| ---- | ----------- |
+| `CreateStream` |  |
+| `UpdateStream` |  |
+| `DeleteStream` |  |
+| `ViewStream` |  |
+| `ViewChanges` |  |
+| `ViewTemplate` |  |
+
+## TestData
+
+| Name | Description |
+| ---- | ----------- |
+| `TestHealthRead` |  |
+| `TestHealthWrite` |  |
 
 ## Accounts
 
@@ -1012,6 +1295,19 @@
 | `DownloadSoftware` |  |
 | `DeleteSoftware` |  |
 
+## Storage
+
+| Name | Description |
+| ---- | ----------- |
+| `ReadBlobs` |  |
+| `WriteBlobs` |  |
+| `ReadAliases` |  |
+| `WriteAliases` |  |
+| `DeleteAliases` |  |
+| `ReadRefs` |  |
+| `WriteRefs` |  |
+| `DeleteRefs` |  |
+
 ## Streams
 
 | Name | Description |
@@ -1036,13 +1332,6 @@
 | `TestHealthRead` |  |
 | `TestHealthWrite` |  |
 
-## Tools
-
-| Name | Description |
-| ---- | ----------- |
-| `DownloadTool` |  |
-| `UploadTool` |  |
-
 ## Accounts
 
 | Name | Description |
@@ -1206,14 +1495,14 @@
 
 | Name | Description |
 | ---- | ----------- |
-| `ReadBlobs` | Ability to read blobs from the storage service |
-| `WriteBlobs` | Ability to write blobs to the storage service |
-| `ReadAliases` | Ability to read aliases from the storage service |
-| `WriteAliases` | Ability to write aliases to the storage service |
-| `DeleteAliases` | Ability to write aliases to the storage service |
-| `ReadRefs` | Ability to read refs from the storage service |
-| `WriteRefs` | Ability to write refs to the storage service |
-| `DeleteRefs` | Ability to delete refs |
+| `ReadBlobs` |  |
+| `WriteBlobs` |  |
+| `ReadAliases` |  |
+| `WriteAliases` |  |
+| `DeleteAliases` |  |
+| `ReadRefs` |  |
+| `WriteRefs` |  |
+| `DeleteRefs` |  |
 
 ## Streams
 
@@ -1243,8 +1532,8 @@
 
 | Name | Description |
 | ---- | ----------- |
-| `DownloadTool` |  |
-| `UploadTool` |  |
+| `DownloadTool` | Ability to download a tool |
+| `UploadTool` | Ability to upload new tool versions |
 
 ## Accounts
 
diff --git a/Engine/Source/Programs/Horde/Docs/Config/Schema/Globals.md b/Engine/Source/Programs/Horde/Docs/Config/Schema/Globals.md
--- a/Engine/Source/Programs/Horde/Docs/Config/Schema/Globals.md
+++ b/Engine/Source/Programs/Horde/Docs/Config/Schema/Globals.md
@@ -59,23 +59,15 @@
 
 Name | Description
 ---- | -----------
-`analytics` | [AnalyticsConfig](#analyticsconfig)<br>Configuration for the analytics plugin
 `build` | [BuildConfig](#buildconfig)<br>Configuration for the build plugin
+`analytics` | [AnalyticsConfig](#analyticsconfig)<br>Configuration for the analytics plugin
+`compute` | [ComputeConfig](#computeconfig)<br>Configuration for the compute plugin
+`secrets` | [SecretsConfig](#secretsconfig)<br>Configuration for the secrets plugin
 `exampleplugin` | [ExamplePluginConfig](#examplepluginconfig)<br>Configuration for the exampleplugin plugin
 `experimental` | [ExperimentalConfig](#experimentalconfig)<br>Configuration for the experimental plugin
-`secrets` | [SecretsConfig](#secretsconfig)<br>Configuration for the secrets plugin
-`compute` | [ComputeConfig](#computeconfig)<br>Configuration for the compute plugin
-`tools` | [ToolsConfig](#toolsconfig)<br>Configuration for the tools plugin
-`symbols` | [SymbolsConfig](#symbolsconfig)<br>Configuration for the symbols plugin
 `storage` | [StorageConfig](#storageconfig)<br>Configuration for the storage plugin
+`symbols` | [SymbolsConfig](#symbolsconfig)<br>Configuration for the symbols plugin
+`tools` | [ToolsConfig](#toolsconfig)<br>Configuration for the tools plugin
 
-## AnalyticsConfig
-
-Config settings for analytics
-
-Name | Description
----- | -----------
-`stores` | [TelemetryStoreConfig](Telemetry.md)`[]`<br>Metrics to aggregate on the Horde server
-
 ## BuildConfig
 
@@ -253,103 +245,13 @@
 `excludeActions` | `string[]`<br>Actions to exclude from the inherited actions
 `extends` | `string[]`<br>Other profiles to extend from
 
-## ExamplePluginConfig
+## AnalyticsConfig
 
-Configuration for the ExamplePlugin plugin
+Config settings for analytics
 
 Name | Description
 ---- | -----------
-`greetingMessage` | `string`<br>The Greeting provided by the example plugin
+`stores` | [TelemetryStoreConfig](Telemetry.md)`[]`<br>Metrics to aggregate on the Horde server
 
-## ExperimentalConfig
-
-Configuration for the experimental plugin
-
-Name | Description
----- | -----------
-`notifications` | [NotificationConfig](Notification.md)`[]`<br>Notifications to be sent on the Horde server
-
-## SecretsConfig
-
-Configuration for the secrets system
-
-Name | Description
----- | -----------
-`secrets` | [SecretConfig](#secretconfig)`[]`<br>List of secrets
-`providerConfigs` | [SecretProviderConfig](#secretproviderconfig)`[]`<br>List of configurations for secret providers
-
-## SecretConfig
-
-Configuration for a secret value
-
-Name | Description
----- | -----------
-`id` | `string`<br>Identifier for this secret
-`data` | `string` `->` `string`<br>Key/value pairs associated with this secret
-`sources` | [ExternalSecretConfig](#externalsecretconfig)`[]`<br>Providers to source key/value pairs from
-`acl` | [AclConfig](#aclconfig)<br>Defines access to this particular secret
-
-## ExternalSecretConfig
-
-Configuration for an external secret provider
-
-Name | Description
----- | -----------
-`providerConfig` | `string`<br>The name to look up a secret provider config. A config is mutually exclusive to and to be used when a secret requires additional options than just the name of the provider
-`provider` | `string`<br>Name of the provider to use
-`format` | [ExternalSecretFormat](#externalsecretformat-enum)<br>Format of the secret
-`key` | `string`<br>Optional key indicating the parameter to set in the resulting data array. Required if  is .
-`path` | `string`<br>Optional value indicating what to fetch from the provider
-
-## ExternalSecretFormat (Enum)
-
-Format describing how to parse external secret values
-
-Name | Description
----- | -----------
-`Text` | Secret is a plain text value which will be stored using the external secret key
-`Json` | Secret is a JSON formatted string containing key/value pairs
-
-## SecretProviderConfig
-
-Common settings object for different secret providers
-
-Name | Description
----- | -----------
-`name` | `string`<br>Name of the secret provider config
-`provider` | `string`<br>Name of the provider
-`hcpVault` | [HcpVaultConfig](#hcpvaultconfig)<br>Configuration for HCP Vault
-`azureKeyVault` | [AzureKeyVaultConfig](#azurekeyvaultconfig)<br>Configuration for Azure Key Vault
-
-## HcpVaultConfig
-
-Configuration for HCP Vault secrets.
-
-Name | Description
----- | -----------
-`preSharedKey` | `string`<br>The Vault authentication token to use when the type of credentials is .<br>More information on this token can be found in the API documentation for X-Vault-Token
-`endPoint` | `string`<br>The address of the Vault server.
-`awsArnRole` | `string`<br>The optional AWS ARN to assume a role when logging into Vault. To be used when the type of credentials is . When not set the AWS SDK default credential search will be used.
-`awsIamServerId` | `string`<br>The value for the X-Vault-AWS-IAM-Server-ID header when the Vault server is configured with iam_server_id_header_value. Most likely to be the Vault server's DNS name. To be used when the type of credentials is .
-`role` | `string`<br>The name of the role in the Vault configuration for the login. This role name can be different to the name of the AWS IAM principal. To be used when the type of credentials is . When not set the name of IAM principal will be used.
-`credentials` | [HcpVaultCredentialsType](#hcpvaultcredentialstype-enum)<br>Type of credentials to use.
-
-## HcpVaultCredentialsType (Enum)
-
-Credentials to use to log into HCP Vault.
-
-Name | Description
----- | -----------
-`PreSharedKey` | Use a pre-shared token provided by Vault for authentication.
-`AwsAuth` | Use AWS IAM credentials to authenticate against the Vault server.
-
-## AzureKeyVaultConfig
-
-Configuration for Azure Key Vault secrets.
-
-Name | Description
----- | -----------
-`vaultUri` | `string`<br>The Uri to the Azure Key Vault.
-
 ## ComputeConfig
 
@@ -550,52 +452,103 @@
 `description` | `string`<br>Human-readable description
 `computeId` | `string`<br>Compute ID for this network (used when allocating compute resources)
 
-## ToolsConfig
+## SecretsConfig
 
-Configuration for the tools system
+Configuration for the secrets system
 
 Name | Description
 ---- | -----------
-`tools` | [ToolConfig](#toolconfig)`[]`<br>Tool configurations
+`secrets` | [SecretConfig](#secretconfig)`[]`<br>List of secrets
+`providerConfigs` | [SecretProviderConfig](#secretproviderconfig)`[]`<br>List of configurations for secret providers
 
-## ToolConfig
+## SecretConfig
 
-Options for configuring a tool
+Configuration for a secret value
 
 Name | Description
 ---- | -----------
-`id` | `string`<br>Unique identifier for the tool
-`name` | `string`<br>Name of the tool
-`description` | `string`<br>Description for the tool
-`category` | `string`<br>Category for the tool. Will cause the tool to be shown in a different tab in the dashboard.
-`group` | `string`<br>Grouping key for different variations of the same tool. The dashboard will show these together.
-`platforms` | `string[]`<br>Platforms for this tool. Takes the form of a NET RID (https://learn.microsoft.com/en-us/dotnet/core/rid-catalog).
-`public` | `boolean`<br>Whether this tool should be exposed for download on a public endpoint without authentication
-`showInUgs` | `boolean`<br>Whether to show this tool for download in the UGS tools menu
-`showInDashboard` | `boolean`<br>Whether to show this tool for download in the dashboard
-`showInToolbox` | `boolean`<br>Whether to show this tool for download in Unreal Toolbox
-`metadata` | `string` `->` `string`<br>Metadata for this tool
-`namespaceId` | `string`<br>Default namespace for new deployments of this tool
-`acl` | [AclConfig](#aclconfig)<br>Permissions for the tool
+`id` | `string`<br>Identifier for this secret
+`data` | `string` `->` `string`<br>Key/value pairs associated with this secret
+`sources` | [ExternalSecretConfig](#externalsecretconfig)`[]`<br>Providers to source key/value pairs from
+`acl` | [AclConfig](#aclconfig)<br>Defines access to this particular secret
 
-## SymbolsConfig
+## ExternalSecretConfig
 
-Configuration for the tools system
+Configuration for an external secret provider
 
 Name | Description
 ---- | -----------
-`stores` | [SymbolStoreConfig](#symbolstoreconfig)`[]`<br>List of symbol stores
+`providerConfig` | `string`<br>The name to look up a secret provider config. A config is mutually exclusive to and to be used when a secret requires additional options than just the name of the provider
+`provider` | `string`<br>Name of the provider to use
+`format` | [ExternalSecretFormat](#externalsecretformat-enum)<br>Format of the secret
+`key` | `string`<br>Optional key indicating the parameter to set in the resulting data array. Required if  is .
+`path` | `string`<br>Optional value indicating what to fetch from the provider
 
-## SymbolStoreConfig
+## ExternalSecretFormat (Enum)
 
-Configuration for a symbol store
+Format describing how to parse external secret values
 
 Name | Description
 ---- | -----------
-`id` | `string`<br>Identifier for this store
-`namespaceId` | `string`<br>Configuration for the symbol store backend
-`public` | `boolean`<br>Whether to make this store available without auth
-`acl` | [AclConfig](#aclconfig)<br>Access to the symbol store
+`Text` | Secret is a plain text value which will be stored using the external secret key
+`Json` | Secret is a JSON formatted string containing key/value pairs
 
+## SecretProviderConfig
+
+Common settings object for different secret providers
+
+Name | Description
+---- | -----------
+`name` | `string`<br>Name of the secret provider config
+`provider` | `string`<br>Name of the provider
+`hcpVault` | [HcpVaultConfig](#hcpvaultconfig)<br>Configuration for HCP Vault
+`azureKeyVault` | [AzureKeyVaultConfig](#azurekeyvaultconfig)<br>Configuration for Azure Key Vault
+
+## HcpVaultConfig
+
+Configuration for HCP Vault secrets.
+
+Name | Description
+---- | -----------
+`preSharedKey` | `string`<br>The Vault authentication token to use when the type of credentials is .<br>More information on this token can be found in the API documentation for X-Vault-Token
+`endPoint` | `string`<br>The address of the Vault server.
+`awsArnRole` | `string`<br>The optional AWS ARN to assume a role when logging into Vault. To be used when the type of credentials is . When not set the AWS SDK default credential search will be used.
+`awsIamServerId` | `string`<br>The value for the X-Vault-AWS-IAM-Server-ID header when the Vault server is configured with iam_server_id_header_value. Most likely to be the Vault server's DNS name. To be used when the type of credentials is .
+`role` | `string`<br>The name of the role in the Vault configuration for the login. This role name can be different to the name of the AWS IAM principal. To be used when the type of credentials is . When not set the name of IAM principal will be used.
+`credentials` | [HcpVaultCredentialsType](#hcpvaultcredentialstype-enum)<br>Type of credentials to use.
+
+## HcpVaultCredentialsType (Enum)
+
+Credentials to use to log into HCP Vault.
+
+Name | Description
+---- | -----------
+`PreSharedKey` | Use a pre-shared token provided by Vault for authentication.
+`AwsAuth` | Use AWS IAM credentials to authenticate against the Vault server.
+
+## AzureKeyVaultConfig
+
+Configuration for Azure Key Vault secrets.
+
+Name | Description
+---- | -----------
+`vaultUri` | `string`<br>The Uri to the Azure Key Vault.
+
+## ExamplePluginConfig
+
+Configuration for the ExamplePlugin plugin
+
+Name | Description
+---- | -----------
+`greetingMessage` | `string`<br>The Greeting provided by the example plugin
+
+## ExperimentalConfig
+
+Configuration for the experimental plugin
+
+Name | Description
+---- | -----------
+`notifications` | [NotificationConfig](Notification.md)`[]`<br>Notifications to be sent on the Horde server
+
 ## StorageConfig
 
@@ -674,3 +627,50 @@
 `gcDelayHrs` | `number`<br>How long to keep newly uploaded orphanned objects before allowing them to be deleted, in hours.
 `enableAliases` | `boolean`<br>Support querying exports by their aliases
 `acl` | [AclConfig](#aclconfig)<br>Access list for this namespace
+
+## SymbolsConfig
+
+Configuration for the tools system
+
+Name | Description
+---- | -----------
+`stores` | [SymbolStoreConfig](#symbolstoreconfig)`[]`<br>List of symbol stores
+
+## SymbolStoreConfig
+
+Configuration for a symbol store
+
+Name | Description
+---- | -----------
+`id` | `string`<br>Identifier for this store
+`namespaceId` | `string`<br>Configuration for the symbol store backend
+`public` | `boolean`<br>Whether to make this store available without auth
+`acl` | [AclConfig](#aclconfig)<br>Access to the symbol store
+
+## ToolsConfig
+
+Configuration for the tools system
+
+Name | Description
+---- | -----------
+`tools` | [ToolConfig](#toolconfig)`[]`<br>Tool configurations
+
+## ToolConfig
+
+Options for configuring a tool
+
+Name | Description
+---- | -----------
+`id` | `string`<br>Unique identifier for the tool
+`name` | `string`<br>Name of the tool
+`description` | `string`<br>Description for the tool
+`category` | `string`<br>Category for the tool. Will cause the tool to be shown in a different tab in the dashboard.
+`group` | `string`<br>Grouping key for different variations of the same tool. The dashboard will show these together.
+`platforms` | `string[]`<br>Platforms for this tool. Takes the form of a NET RID (https://learn.microsoft.com/en-us/dotnet/core/rid-catalog).
+`public` | `boolean`<br>Whether this tool should be exposed for download on a public endpoint without authentication
+`showInUgs` | `boolean`<br>Whether to show this tool for download in the UGS tools menu
+`showInDashboard` | `boolean`<br>Whether to show this tool for download in the dashboard
+`showInToolbox` | `boolean`<br>Whether to show this tool for download in Unreal Toolbox
+`metadata` | `string` `->` `string`<br>Metadata for this tool
+`namespaceId` | `string`<br>Default namespace for new deployments of this tool
+`acl` | [AclConfig](#aclconfig)<br>Permissions for the tool
diff --git a/Engine/Source/Programs/Horde/Docs/Deployment/ServerSettings.md b/Engine/Source/Programs/Horde/Docs/Deployment/ServerSettings.md
--- a/Engine/Source/Programs/Horde/Docs/Deployment/ServerSettings.md
+++ b/Engine/Source/Programs/Horde/Docs/Deployment/ServerSettings.md
@@ -127,52 +127,15 @@
 
 Name | Description
 ---- | -----------
-`analytics` | [AnalyticsServerConfig](#analyticsserverconfig)<br>Configuration for the analytics plugin
 `build` | [BuildServerConfig](#buildserverconfig)<br>Configuration for the build plugin
+`analytics` | [AnalyticsServerConfig](#analyticsserverconfig)<br>Configuration for the analytics plugin
+`compute` | [ComputeServerConfig](#computeserverconfig)<br>Configuration for the compute plugin
+`secrets` | [SecretsServerConfig](#secretsserverconfig)<br>Configuration for the secrets plugin
 `exampleplugin` | [ExamplePluginServerConfig](#examplepluginserverconfig)<br>Configuration for the exampleplugin plugin
 `experimental` | [ExperimentalServerConfig](#experimentalserverconfig)<br>Configuration for the experimental plugin
-`secrets` | [SecretsServerConfig](#secretsserverconfig)<br>Configuration for the secrets plugin
-`compute` | [ComputeServerConfig](#computeserverconfig)<br>Configuration for the compute plugin
-`tools` | [ToolsServerConfig](#toolsserverconfig)<br>Configuration for the tools plugin
-`symbols` | [PluginServerConfig](#pluginserverconfig)<br>Configuration for the symbols plugin
 `storage` | [StorageServerConfig](#storageserverconfig)<br>Configuration for the storage plugin
+`symbols` | [PluginServerConfig](#pluginserverconfig)<br>Configuration for the symbols plugin
+`tools` | [ToolsServerConfig](#toolsserverconfig)<br>Configuration for the tools plugin
 
-## AnalyticsServerConfig
-
-Server configuration for the analytics system
-
-Name | Description
----- | -----------
-`sinks` | [TelemetrySinkConfig](#telemetrysinkconfig)<br>Settings for the various telemetry sinks
-`enabled` | `boolean`<br>Whether the plugin should be enabled or not
-
-## TelemetrySinkConfig
-
-Telemetry sinks
-
-Name | Description
----- | -----------
-`epic` | [EpicTelemetryConfig](#epictelemetryconfig)<br>Settings for the Epic telemetry sink
-`mongo` | [MongoTelemetryConfig](#mongotelemetryconfig)<br>Settings for the MongoDB telemetry sink
-
-## EpicTelemetryConfig
-
-Configuration for the telemetry sink
-
-Name | Description
----- | -----------
-`url` | `string`<br>Base URL for the telemetry server
-`appId` | `string`<br>Application name to send in the event messages
-`enabled` | `boolean`<br>Whether to enable this sink
-
-## MongoTelemetryConfig
-
-Configuration for the telemetry sink
-
-Name | Description
----- | -----------
-`retainDays` | `number`<br>Number of days worth of telmetry events to keep
-`enabled` | `boolean`<br>Whether to enable this sink
-
 ## BuildServerConfig
 
@@ -318,30 +281,42 @@
 `maxChildCount` | `integer`<br>Maximum number of children in each node
 `sliceThreshold` | `integer`<br>Threshold hash value for splitting interior nodes
 
-## ExamplePluginServerConfig
+## AnalyticsServerConfig
 
-Server configuration for the ExamplePlugin
+Server configuration for the analytics system
 
 Name | Description
 ---- | -----------
+`sinks` | [TelemetrySinkConfig](#telemetrysinkconfig)<br>Settings for the various telemetry sinks
 `enabled` | `boolean`<br>Whether the plugin should be enabled or not
 
-## ExperimentalServerConfig
+## TelemetrySinkConfig
 
-Server configuration for the experimental plugin
+Telemetry sinks
 
 Name | Description
 ---- | -----------
-`enabled` | `boolean`<br>Whether the plugin should be enabled or not
+`epic` | [EpicTelemetryConfig](#epictelemetryconfig)<br>Settings for the Epic telemetry sink
+`mongo` | [MongoTelemetryConfig](#mongotelemetryconfig)<br>Settings for the MongoDB telemetry sink
 
-## SecretsServerConfig
+## EpicTelemetryConfig
 
-Static configuration for the secrets plugin
+Configuration for the telemetry sink
 
 Name | Description
 ---- | -----------
-`withAws` | `boolean`<br>Whether to enable Amazon Web Services (AWS) specific features
-`enabled` | `boolean`<br>Whether the plugin should be enabled or not
+`url` | `string`<br>Base URL for the telemetry server
+`appId` | `string`<br>Application name to send in the event messages
+`enabled` | `boolean`<br>Whether to enable this sink
 
+## MongoTelemetryConfig
+
+Configuration for the telemetry sink
+
+Name | Description
+---- | -----------
+`retainDays` | `number`<br>Number of days worth of telmetry events to keep
+`enabled` | `boolean`<br>Whether to enable this sink
+
 ## ComputeServerConfig
 
@@ -391,82 +366,26 @@
 `ComputeQueueAwsMetric` | A no-op strategy that reports metrics to let an external AWS auto-scaling policy scale the fleet
 `LeaseUtilizationAwsMetric` | A no-op strategy that reports metrics to let an external AWS auto-scaling policy scale the fleet
 
-## ToolsServerConfig
+## SecretsServerConfig
 
-Server configuration for bundled tools
+Static configuration for the secrets plugin
 
 Name | Description
 ---- | -----------
-`bundledTools` | [BundledToolConfig](#bundledtoolconfig)`[]`<br>Tools bundled along with the server. Data for each tool can be produced using the 'bundle create' command, and should be stored in the Tools directory.
+`withAws` | `boolean`<br>Whether to enable Amazon Web Services (AWS) specific features
 `enabled` | `boolean`<br>Whether the plugin should be enabled or not
 
-## BundledToolConfig
+## ExamplePluginServerConfig
 
-Configuration for a tool bundled alongsize the server
+Server configuration for the ExamplePlugin
 
 Name | Description
 ---- | -----------
-`version` | `string`<br>Version string for the current tool data
-`refName` | `string`<br>Ref name in the tools directory
-`dataDir` | `string`<br>Directory containing blob data for this tool. If empty, the tools/{id} folder next to the server will be used.
-`id` | `string`<br>Unique identifier for the tool
-`name` | `string`<br>Name of the tool
-`description` | `string`<br>Description for the tool
-`category` | `string`<br>Category for the tool. Will cause the tool to be shown in a different tab in the dashboard.
-`group` | `string`<br>Grouping key for different variations of the same tool. The dashboard will show these together.
-`platforms` | `string[]`<br>Platforms for this tool. Takes the form of a NET RID (https://learn.microsoft.com/en-us/dotnet/core/rid-catalog).
-`public` | `boolean`<br>Whether this tool should be exposed for download on a public endpoint without authentication
-`showInUgs` | `boolean`<br>Whether to show this tool for download in the UGS tools menu
-`showInDashboard` | `boolean`<br>Whether to show this tool for download in the dashboard
-`showInToolbox` | `boolean`<br>Whether to show this tool for download in Unreal Toolbox
-`metadata` | `string` `->` `string`<br>Metadata for this tool
-`namespaceId` | `string`<br>Default namespace for new deployments of this tool
-`acl` | [AclConfig](#aclconfig)<br>Permissions for the tool
+`enabled` | `boolean`<br>Whether the plugin should be enabled or not
 
-## AclConfig
+## ExperimentalServerConfig
 
-Parameters to update an ACL
+Server configuration for the experimental plugin
 
 Name | Description
 ---- | -----------
-`entries` | [AclEntryConfig](#aclentryconfig)`[]`<br>Entries to replace the existing ACL
-`profiles` | [AclProfileConfig](#aclprofileconfig)`[]`<br>Defines profiles which allow grouping sets of actions into named collections
-`inherit` | `boolean`<br>Whether to inherit permissions from the parent ACL
-`exceptions` | `string[]`<br>List of exceptions to the inherited setting
-
-## AclEntryConfig
-
-Individual entry in an ACL
-
-Name | Description
----- | -----------
-`claim` | [AclClaimConfig](#aclclaimconfig)<br>Name of the user or group
-`actions` | `string[]`<br>Array of actions to allow
-`profiles` | `string[]`<br>List of profiles to grant
-
-## AclClaimConfig
-
-New claim to create
-
-Name | Description
----- | -----------
-`type` | `string`<br>The claim type
-`value` | `string`<br>The claim value
-
-## AclProfileConfig
-
-Configuration for an ACL profile. This defines a preset group of actions which can be given to a user via an ACL entry.
-
-Name | Description
----- | -----------
-`id` | `string`<br>Identifier for this profile
-`actions` | `string[]`<br>Actions to include
-`excludeActions` | `string[]`<br>Actions to exclude from the inherited actions
-`extends` | `string[]`<br>Other profiles to extend from
-
-## PluginServerConfig
-
-Base class for plugin server config objects
-
-Name | Description
----- | -----------
@@ -530,3 +449,84 @@
 `Profile` | Read credentials from the  profile in the AWS config file
 `AssumeRole` | Assume a particular role. Should specify ARN in
 `AssumeRoleWebIdentity` | Assume a particular role using the current environment variables.
+
+## PluginServerConfig
+
+Base class for plugin server config objects
+
+Name | Description
+---- | -----------
+`enabled` | `boolean`<br>Whether the plugin should be enabled or not
+
+## ToolsServerConfig
+
+Server configuration for bundled tools
+
+Name | Description
+---- | -----------
+`bundledTools` | [BundledToolConfig](#bundledtoolconfig)`[]`<br>Tools bundled along with the server. Data for each tool can be produced using the 'bundle create' command, and should be stored in the Tools directory.
+`enabled` | `boolean`<br>Whether the plugin should be enabled or not
+
+## BundledToolConfig
+
+Configuration for a tool bundled alongsize the server
+
+Name | Description
+---- | -----------
+`version` | `string`<br>Version string for the current tool data
+`refName` | `string`<br>Ref name in the tools directory
+`dataDir` | `string`<br>Directory containing blob data for this tool. If empty, the tools/{id} folder next to the server will be used.
+`id` | `string`<br>Unique identifier for the tool
+`name` | `string`<br>Name of the tool
+`description` | `string`<br>Description for the tool
+`category` | `string`<br>Category for the tool. Will cause the tool to be shown in a different tab in the dashboard.
+`group` | `string`<br>Grouping key for different variations of the same tool. The dashboard will show these together.
+`platforms` | `string[]`<br>Platforms for this tool. Takes the form of a NET RID (https://learn.microsoft.com/en-us/dotnet/core/rid-catalog).
+`public` | `boolean`<br>Whether this tool should be exposed for download on a public endpoint without authentication
+`showInUgs` | `boolean`<br>Whether to show this tool for download in the UGS tools menu
+`showInDashboard` | `boolean`<br>Whether to show this tool for download in the dashboard
+`showInToolbox` | `boolean`<br>Whether to show this tool for download in Unreal Toolbox
+`metadata` | `string` `->` `string`<br>Metadata for this tool
+`namespaceId` | `string`<br>Default namespace for new deployments of this tool
+`acl` | [AclConfig](#aclconfig)<br>Permissions for the tool
+
+## AclConfig
+
+Parameters to update an ACL
+
+Name | Description
+---- | -----------
+`entries` | [AclEntryConfig](#aclentryconfig)`[]`<br>Entries to replace the existing ACL
+`profiles` | [AclProfileConfig](#aclprofileconfig)`[]`<br>Defines profiles which allow grouping sets of actions into named collections
+`inherit` | `boolean`<br>Whether to inherit permissions from the parent ACL
+`exceptions` | `string[]`<br>List of exceptions to the inherited setting
+
+## AclEntryConfig
+
+Individual entry in an ACL
+
+Name | Description
+---- | -----------
+`claim` | [AclClaimConfig](#aclclaimconfig)<br>Name of the user or group
+`actions` | `string[]`<br>Array of actions to allow
+`profiles` | `string[]`<br>List of profiles to grant
+
+## AclClaimConfig
+
+New claim to create
+
+Name | Description
+---- | -----------
+`type` | `string`<br>The claim type
+`value` | `string`<br>The claim value
+
+## AclProfileConfig
+
+Configuration for an ACL profile. This defines a preset group of actions which can be given to a user via an ACL entry.
+
+Name | Description
+---- | -----------
+`id` | `string`<br>Identifier for this profile
+`actions` | `string[]`<br>Actions to include
+`excludeActions` | `string[]`<br>Actions to exclude from the inherited actions
+`extends` | `string[]`<br>Other profiles to extend from
