Horde: Set a default telemetry store for compute

diff --git a/Engine/Source/Programs/Horde/Plugins/Compute/HordeServer.Compute/ComputeConfig.cs b/Engine/Source/Programs/Horde/Plugins/Compute/HordeServer.Compute/ComputeConfig.cs
--- a/Engine/Source/Programs/Horde/Plugins/Compute/HordeServer.Compute/ComputeConfig.cs
+++ b/Engine/Source/Programs/Horde/Plugins/Compute/HordeServer.Compute/ComputeConfig.cs
@@ -56,7 +56,7 @@ namespace HordeServer
 		/// <summary>
 		/// Telemetry store for lease and compute telemetry data
 		/// </summary>
-		public TelemetryStoreId TelemetryStoreId { get; set; }
+		public TelemetryStoreId TelemetryStoreId { get; set; } = TelemetryStoreId.Default;
 
 		/// <summary>
 		/// List of costs of a particular agent type
