Horde: Make tool deployment retention limit configurable

Adds a MaxDeployments property to ToolConfig (defaulting to 5, preserving existing behavior) to allow per-tool configuration of how many deployments are retained.

diff --git a/Engine/Source/Programs/Horde/Plugins/Tools/HordeServer.Tools.Tests/Tools/ToolTests.cs b/Engine/Source/Programs/Horde/Plugins/Tools/HordeServer.Tools.Tests/Tools/ToolTests.cs
--- a/Engine/Source/Programs/Horde/Plugins/Tools/HordeServer.Tools.Tests/Tools/ToolTests.cs
+++ b/Engine/Source/Programs/Horde/Plugins/Tools/HordeServer.Tools.Tests/Tools/ToolTests.cs
@@ -140,6 +140,52 @@ namespace HordeServer.Tests.Tools
 			}
 		}
 
+		[TestMethod]
+		public async Task MaxDeploymentsExceededPurgesOldestAsync()
+		{
+			// Reconfigure the tool with a low MaxDeployments limit
+			const int MaxDeployments = 3;
+
+			StorageConfig storageConfig = new StorageConfig();
+			storageConfig.Backends.Clear();
+			storageConfig.Backends.Add(new BackendConfig { Id = new BackendId("tools-backend"), Type = StorageBackendType.Memory });
+			storageConfig.Namespaces.Clear();
+			storageConfig.Namespaces.Add(new NamespaceConfig { Id = ToolConfig.DefaultNamespaceId, Backend = new BackendId("tools-backend") });
+
+			ToolsConfig toolsConfig = new ToolsConfig();
+			toolsConfig.Tools.Add(new ToolConfig(_toolId) { Name = "UnrealGameSync", Description = "Tool for syncing content from source control", MaxDeployments = MaxDeployments });
+
+			GlobalConfig globalConfig = new GlobalConfig();
+			globalConfig.Plugins.AddStorageConfig(storageConfig);
+			globalConfig.Plugins.AddToolsConfig(toolsConfig);
+			await SetConfigAsync(globalConfig);
+
+			IToolCollection collection = ServiceProvider.GetRequiredService<IToolCollection>();
+			ITool tool = Deref(await collection.GetAsync(_toolId));
+
+			// Create more deployments than the limit
+			List<string> versions = new List<string>();
+			for (int i = 0; i < MaxDeployments + 2; i++)
+			{
+				string version = $"{i + 1}.0";
+				versions.Add(version);
+
+				byte[] zipData = await CreateZipFileDataAsync($"file-{version}.txt", $"content-{version}");
+				using MemoryStream stream = new MemoryStream(zipData);
+				tool = Deref(await tool.CreateDeploymentAsync(new ToolDeploymentConfig { Version = version }, stream, CancellationToken.None));
+			}
+
+			// Verify only MaxDeployments remain and they are the most recent ones
+			tool = Deref(await collection.GetAsync(_toolId));
+			Assert.AreEqual(MaxDeployments, tool.Deployments.Count);
+
+			List<string> expectedVersions = versions.Skip(versions.Count - MaxDeployments).ToList();
+			for (int i = 0; i < MaxDeployments; i++)
+			{
+				Assert.AreEqual(expectedVersions[i], tool.Deployments[i].Version);
+			}
+		}
+
 		public static async Task<byte[]> CreateZipFileDataAsync(string fileName, string fileData)
 		{
 			using MemoryStream stream = new();
diff --git a/Engine/Source/Programs/Horde/Plugins/Tools/HordeServer.Tools/Tools/ToolCollection.cs b/Engine/Source/Programs/Horde/Plugins/Tools/HordeServer.Tools/Tools/ToolCollection.cs
--- a/Engine/Source/Programs/Horde/Plugins/Tools/HordeServer.Tools/Tools/ToolCollection.cs
+++ b/Engine/Source/Programs/Horde/Plugins/Tools/HordeServer.Tools/Tools/ToolCollection.cs
@@ -380,7 +380,7 @@ namespace HordeServer.Tools
 			ToolDocument? newTool = (ToolDocument)tool;
 			for (; ; )
 			{
-				newTool = await TryAddDeploymentAsync(newTool, deployment, cancellationToken);
+				newTool = await TryAddDeploymentAsync(newTool, deployment, toolConfig.MaxDeployments, cancellationToken);
 				if (newTool != null)
 				{
 					return newTool;
@@ -394,13 +394,14 @@ namespace HordeServer.Tools
 			}
 		}
 
-		async ValueTask<ToolDocument?> TryAddDeploymentAsync(ToolDocument tool, ToolDeploymentDocument deployment, CancellationToken cancellationToken)
+		async ValueTask<ToolDocument?> TryAddDeploymentAsync(ToolDocument tool, ToolDeploymentDocument deployment, int maxDeploymentCount, CancellationToken cancellationToken)
 		{
+			maxDeploymentCount = Math.Max(1, maxDeploymentCount);
+
 			ToolDocument? newTool = tool;
 
 			// If there are already a maximum number of deployments, remove the oldest one
-			const int MaxDeploymentCount = 5;
-			while (newTool.Deployments.Count >= MaxDeploymentCount)
+			while (newTool.Deployments.Count >= maxDeploymentCount)
 			{
 				newTool = await UpdateAsync(newTool, Builders<ToolDocument>.Update.PopFirst(x => x.Deployments), cancellationToken);
 				if (newTool == null)
diff --git a/Engine/Source/Programs/Horde/Plugins/Tools/HordeServer.Tools/Tools/ToolConfig.cs b/Engine/Source/Programs/Horde/Plugins/Tools/HordeServer.Tools/Tools/ToolConfig.cs
--- a/Engine/Source/Programs/Horde/Plugins/Tools/HordeServer.Tools/Tools/ToolConfig.cs
+++ b/Engine/Source/Programs/Horde/Plugins/Tools/HordeServer.Tools/Tools/ToolConfig.cs
@@ -86,6 +86,11 @@ namespace HordeServer.Tools
 		/// </summary>
 		public NamespaceId NamespaceId { get; set; } = DefaultNamespaceId;
 
+		/// <summary>
+		/// Maximum number of deployments to retain for this tool. When exceeded, the oldest deployments are removed.
+		/// </summary>
+		public int MaxDeployments { get; set; } = 5;
+
 		/// <summary>
 		/// Permissions for the tool
 		/// </summary>
