Horde: Handle missing log blobs gracefully in step history events
When log blob data has been garbage collected from storage, the
GetEventsAsync and GetBatchEventsAsync endpoints now catch
StorageException and skip the unavailable event instead of
returning a 500 error to the dashboard.

diff --git a/Engine/Source/Programs/Horde/Plugins/Compute/HordeServer.Compute/Logs/LogsController.cs b/Engine/Source/Programs/Horde/Plugins/Compute/HordeServer.Compute/Logs/LogsController.cs
--- a/Engine/Source/Programs/Horde/Plugins/Compute/HordeServer.Compute/Logs/LogsController.cs
+++ b/Engine/Source/Programs/Horde/Plugins/Compute/HordeServer.Compute/Logs/LogsController.cs
@@ -13,6 +13,7 @@ using HordeServer.Storage;
 using HordeServer.Utilities;
 using Microsoft.AspNetCore.Authorization;
 using Microsoft.AspNetCore.Mvc;
+using Microsoft.Extensions.Logging;
 using Microsoft.Extensions.Options;
 using MongoDB.Bson;
 
@@ -47,17 +48,19 @@ namespace HordeServer.Logs
 		private readonly IEnumerable<ILogExtAuthProvider> _authProviders;
 		private readonly IEnumerable<ILogExtIssueProvider> _issueProviders;
 		private readonly IOptionsSnapshot<ComputeConfig> _computeConfig;
+		private readonly ILogger _logger;
 
 		/// <summary>
 		/// Constructor
 		/// </summary>
-		public LogsController(ILogCollection logCollection, StorageService storageService, IEnumerable<ILogExtAuthProvider> authProviders, IEnumerable<ILogExtIssueProvider> issueProviders, IOptionsSnapshot<ComputeConfig> computeConfig)
+		public LogsController(ILogCollection logCollection, StorageService storageService, IEnumerable<ILogExtAuthProvider> authProviders, IEnumerable<ILogExtIssueProvider> issueProviders, IOptionsSnapshot<ComputeConfig> computeConfig, ILogger<LogsController> logger)
 		{
 			_logCollection = logCollection;
 			_storageService = storageService;
 			_authProviders = authProviders;
 			_issueProviders = issueProviders;
 			_computeConfig = computeConfig;
+			_logger = logger;
 		}
 
 		/// <summary>
@@ -355,7 +358,16 @@ namespace HordeServer.Logs
 			List<GetLogEventResponse> responses = new List<GetLogEventResponse>();
 			foreach (ILogAnchor anchor in anchors)
 			{
-				ILogEventData logEventData = await anchor.GetDataAsync(cancellationToken);
+				ILogEventData logEventData;
+				try
+				{
+					logEventData = await anchor.GetDataAsync(cancellationToken);
+				}
+				catch (StorageException ex)
+				{
+					_logger.LogWarning(ex, "Unable to read log event data for log {LogId} at line {LineIndex} (data may have been garbage collected)", logId, anchor.LineIndex);
+					continue;
+				}
 
 				int? issueId = null;
 				if (anchor.SpanId != null && !spanIdToIssueId.TryGetValue(anchor.SpanId.Value, out issueId))
@@ -421,7 +433,16 @@ namespace HordeServer.Logs
 
 				foreach (ILogAnchor anchor in anchorsToFetch)
 				{
-					ILogEventData logEventData = await anchor.GetDataAsync(cancellationToken);
+					ILogEventData logEventData;
+					try
+					{
+						logEventData = await anchor.GetDataAsync(cancellationToken);
+					}
+					catch (StorageException ex)
+					{
+						_logger.LogWarning(ex, "Unable to read log event data for log {LogId} at line {LineIndex} (data may have been garbage collected)", logId, anchor.LineIndex);
+						continue;
+					}
 
 					int? issueId = null;
 					if (anchor.SpanId != null && !spanIdToIssueId.TryGetValue(anchor.SpanId.Value, out issueId))
