Horde: Fix error 'Change has shelved files associated with it and can't be deleted' during conforms

#Synced-CL 0

diff --git a/Engine/Source/Programs/Horde/Drivers/JobDriver/Utility/WorkspaceInfo.cs b/Engine/Source/Programs/Horde/Drivers/JobDriver/Utility/WorkspaceInfo.cs
--- a/Engine/Source/Programs/Horde/Drivers/JobDriver/Utility/WorkspaceInfo.cs
+++ b/Engine/Source/Programs/Horde/Drivers/JobDriver/Utility/WorkspaceInfo.cs
@@ -504,7 +504,7 @@ namespace JobDriver.Utility
 				if (records.Count > 0 && records[0].Files.Count > 0)
 				{
 					logger.LogInformation("Deleting shelved files in changelist {Change}", pendingChange.Number);
-					await perforce.DeleteChangeAsync(DeleteChangeOptions.None, pendingChange.Number, cancellationToken);
+					await perforce.DeleteShelvedFilesAsync(pendingChange.Number, FileSpecList.Any, cancellationToken);
 				}
 
 				// delete the changelist
