Horde: Add ResolveSecret to TokenConfig.ClientId and ClientSecret

#Synced-CL 0

diff --git a/Engine/Source/Programs/Horde/Plugins/Build/HordeServer.Build.Tests/Streams/TokenConfigSecretsTests.cs b/Engine/Source/Programs/Horde/Plugins/Build/HordeServer.Build.Tests/Streams/TokenConfigSecretsTests.cs
new file mode 100644
--- /dev/null
+++ b/Engine/Source/Programs/Horde/Plugins/Build/HordeServer.Build.Tests/Streams/TokenConfigSecretsTests.cs
@@ -0,0 +1,111 @@
+// Copyright Epic Games, Inc. All Rights Reserved.
+
+using EpicGames.Horde.Secrets;
+using EpicGames.Horde.Streams;
+using HordeServer.Projects;
+using HordeServer.Secrets;
+using HordeServer.Server;
+using HordeServer.Streams;
+
+namespace HordeServer.Tests.Streams;
+
+[TestClass]
+public class TokenConfigSecretsTests : BuildTestSetup
+{
+	public TokenConfigSecretsTests()
+	{
+		AddPlugin<SecretsPlugin>();
+	}
+
+	[TestInitialize]
+	public async Task SetupAsync()
+	{
+		GlobalConfig globalConfig = new() { ResolveSecrets = true };
+		globalConfig.Plugins.AddBuildConfig(new BuildConfig());
+		globalConfig.Plugins.AddSecretsConfig(new SecretsConfig());
+		await SetConfigAsync(globalConfig);
+	}
+
+	private async Task AddTokenConfigAsync(TokenConfig tokenConfig, SecretConfig? secret)
+	{
+		await UpdateConfigAsync(globalConfig =>
+		{
+			if (secret != null)
+			{
+				globalConfig.Plugins.GetSecretsConfig().Secrets.Add(secret);
+			}
+			globalConfig.Plugins.GetBuildConfig().Projects.Add(new ProjectConfig
+			{
+				Streams =
+				[
+					new StreamConfig
+					{
+						Id = new StreamId("test-stream"),
+						Tokens = [tokenConfig]
+					}
+				]
+			});
+		});
+	}
+
+	[TestMethod]
+	public async Task ResolveTokenClientSecretWithSecretAsync()
+	{
+		TokenConfig tokenConfig = new()
+		{
+			Url = new Uri("https://example.com/token"),
+			ClientId = "horde:secret:tokensecret.ClientId",
+			ClientSecret = "horde:secret:tokensecret.ClientSecret",
+			EnvVar = "TOKEN_ENV"
+		};
+		SecretConfig secret = new()
+		{
+			Id = new SecretId("tokensecret"),
+			Data = new Dictionary<string, string> { { "ClientId", "ResolvedClientId" }, { "ClientSecret", "SuperSecretValue" } }
+		};
+		await AddTokenConfigAsync(tokenConfig, secret);
+		Assert.AreEqual("ResolvedClientId", tokenConfig.ClientId);
+		Assert.AreEqual("SuperSecretValue", tokenConfig.ClientSecret);
+	}
+
+	[TestMethod]
+	public async Task ResolveTokenClientSecretWithoutSecretAsync()
+	{
+		TokenConfig tokenConfig = new()
+		{
+			Url = new Uri("https://example.com/token"),
+			ClientId = "my-client",
+			ClientSecret = "plain-text-secret",
+			EnvVar = "TOKEN_ENV"
+		};
+		await AddTokenConfigAsync(tokenConfig, null);
+		Assert.AreEqual("plain-text-secret", tokenConfig.ClientSecret);
+		Assert.AreEqual("my-client", tokenConfig.ClientId);
+	}
+
+	[TestMethod]
+	[DataRow("horde:secret:notfound.property", DisplayName = "Resolve Token Secret, secret not found")]
+	[DataRow("horde:secret:tokensecret.notfound", DisplayName = "Resolve Token Secret, property not found")]
+	public async Task ResolveTokenClientSecretNotFoundAsync(string value)
+	{
+		TokenConfig tokenConfig = new()
+		{
+			Url = new Uri("https://example.com/token"),
+			ClientId = "my-client",
+			ClientSecret = value,
+			EnvVar = "TOKEN_ENV"
+		};
+		SecretConfig secret = new() { Id = new SecretId("tokensecret") };
+		await Assert.ThrowsExactlyAsync<KeyNotFoundException>(async () =>
+		{
+			try
+			{
+				await AddTokenConfigAsync(tokenConfig, secret);
+			}
+			catch (AggregateException ae)
+			{
+				throw ae.InnerExceptions[0];
+			}
+		});
+	}
+}
diff --git a/Engine/Source/Programs/Horde/Plugins/Build/HordeServer.Build/Streams/StreamConfig.cs b/Engine/Source/Programs/Horde/Plugins/Build/HordeServer.Build/Streams/StreamConfig.cs
--- a/Engine/Source/Programs/Horde/Plugins/Build/HordeServer.Build/Streams/StreamConfig.cs
+++ b/Engine/Source/Programs/Horde/Plugins/Build/HordeServer.Build/Streams/StreamConfig.cs
@@ -985,12 +985,14 @@ namespace HordeServer.Streams
 		/// Client id to use to request a new token
 		/// </summary>
 		[Required]
+		[ResolveSecret]
 		public string ClientId { get; set; } = String.Empty;
 
 		/// <summary>
 		/// Client secret to request a new access token
 		/// </summary>
 		[Required]
+		[ResolveSecret]
 		public string ClientSecret { get; set; } = String.Empty;
 
 		/// <summary>
