Horde: Remove conform task log message

#rnx
#Synced-CL 0

diff --git a/Engine/Source/Programs/Horde/Drivers/JobDriver/Execution/ConformExecutor.cs b/Engine/Source/Programs/Horde/Drivers/JobDriver/Execution/ConformExecutor.cs
--- a/Engine/Source/Programs/Horde/Drivers/JobDriver/Execution/ConformExecutor.cs
+++ b/Engine/Source/Programs/Horde/Drivers/JobDriver/Execution/ConformExecutor.cs
@@ -80,8 +80,7 @@ namespace JobDriver.Execution
 
 				if (managedWorkspaces.Count > 0 && perforceWorkspaces.Count > 0)
 				{
-					_logger.LogInformation("Agent contains a mix of workspace materializers:\n{ManagedWorkspaces}\n{PerforceWorkspaces}",
-						managedWorkspaces, perforceWorkspaces);
+					_logger.LogInformation("Agent contains a mix of workspace materializers");
 				}
 
 				foreach (RpcAgentWorkspace raw in managedWorkspaces)
