Horde: Feature flag bisection - disabled by default

diff --git a/Engine/Source/Programs/Horde/HordeDashboard/src/components/UserHome.tsx b/Engine/Source/Programs/Horde/HordeDashboard/src/components/UserHome.tsx
--- a/Engine/Source/Programs/Horde/HordeDashboard/src/components/UserHome.tsx
+++ b/Engine/Source/Programs/Horde/HordeDashboard/src/components/UserHome.tsx
@@ -1036,6 +1036,9 @@ const bisectionHandler = new BisectionHandler();
 
 const BisectionPanel: React.FC = observer(() => {
 
+   // Bisection is currently disabled via feature flag
+   return null;
+
    const { hordeClasses } = getHordeStyling();
 
    useEffect(() => {
diff --git a/Engine/Source/Programs/Horde/HordeDashboard/src/components/jobDetailsV2/JobDetailBisection.tsx b/Engine/Source/Programs/Horde/HordeDashboard/src/components/jobDetailsV2/JobDetailBisection.tsx
--- a/Engine/Source/Programs/Horde/HordeDashboard/src/components/jobDetailsV2/JobDetailBisection.tsx
+++ b/Engine/Source/Programs/Horde/HordeDashboard/src/components/jobDetailsV2/JobDetailBisection.tsx
@@ -176,6 +176,9 @@ JobDetailsV2.registerDataView("StepBisectionView", (details: JobDetailsV2) => ne
 
 export const BisectionPanel: React.FC<{ jobDetails: JobDetailsV2, stepId?: string }> = observer(({ jobDetails, stepId }) => {
 
+   // Bisection is currently disabled via feature flag
+   return null;
+
    const dataView = jobDetails.getDataView<StepBisectionView>("StepBisectionView");
 
    const { hordeClasses } = getHordeStyling();
diff --git a/Engine/Source/Programs/Horde/HordeDashboard/src/components/jobDetailsV2/JobOperationsBar.tsx b/Engine/Source/Programs/Horde/HordeDashboard/src/components/jobDetailsV2/JobOperationsBar.tsx
--- a/Engine/Source/Programs/Horde/HordeDashboard/src/components/jobDetailsV2/JobOperationsBar.tsx
+++ b/Engine/Source/Programs/Horde/HordeDashboard/src/components/jobDetailsV2/JobOperationsBar.tsx
@@ -416,11 +416,8 @@ const StepOperations: React.FC<{ jobDetails: JobDetailsV2, stepId: string }> = o
 
    const canRunDisabled = !step?.allowRetry || !!step.retriedByUserInfo || !!jobData.abortedByUserInfo || !jobDetails.template?.canRun || tempStorageExpired;
    const canTryFix = jobDetails.template?.allowPreflights && !!jobDetails.template?.canRun;
-   let canBisect = (step.outcome === JobStepOutcome.Failure || step.outcome === JobStepOutcome.Warnings) && !!jobDetails.template?.allowPreflights && !!jobDetails.template?.canRun;
-
-   if (getSiteConfig().environment === "production") {
-      canBisect = false;
-   }
+   // Bisection is currently disabled via feature flag
+   const canBisect = false;
 
    const opsList: IContextualMenuItem[] = [];
 
diff --git a/Engine/Source/Programs/Horde/Plugins/Build/HordeServer.Build/BuildPlugin.cs b/Engine/Source/Programs/Horde/Plugins/Build/HordeServer.Build/BuildPlugin.cs
--- a/Engine/Source/Programs/Horde/Plugins/Build/HordeServer.Build/BuildPlugin.cs
+++ b/Engine/Source/Programs/Horde/Plugins/Build/HordeServer.Build/BuildPlugin.cs
@@ -203,7 +203,10 @@ namespace HordeServer
 			if (_serverInfo.IsRunModeActive(RunMode.Worker) && !_serverInfo.ReadOnlyMode)
 			{
 				services.AddHostedService<AgentReportService>();
-				services.AddHostedService<BisectService>();
+				if (!_staticConfig.DisableBisection)
+				{
+					services.AddHostedService<BisectService>();
+				}
 				services.AddHostedService(provider => provider.GetRequiredService<IssueService>());
 				services.AddHostedService<IssueReportService>();
 				services.AddHostedService<IssueTagService>();
diff --git a/Engine/Source/Programs/Horde/Plugins/Build/HordeServer.Build/BuildServerConfig.cs b/Engine/Source/Programs/Horde/Plugins/Build/HordeServer.Build/BuildServerConfig.cs
--- a/Engine/Source/Programs/Horde/Plugins/Build/HordeServer.Build/BuildServerConfig.cs
+++ b/Engine/Source/Programs/Horde/Plugins/Build/HordeServer.Build/BuildServerConfig.cs
@@ -80,6 +80,11 @@ namespace HordeServer
 		/// </summary>
 		public string? DeviceReportChannel { get; set; }
 
+		/// <summary>
+		/// Whether to disable bisection tasks
+		/// </summary>
+		public bool DisableBisection { get; set; } = true;
+
 		/// <summary>
 		/// Whether to run scheduled jobs.
 		/// </summary>
diff --git a/Engine/Source/Programs/Horde/Plugins/Build/HordeServer.Build/Jobs/Bisect/BisectController.cs b/Engine/Source/Programs/Horde/Plugins/Build/HordeServer.Build/Jobs/Bisect/BisectController.cs
--- a/Engine/Source/Programs/Horde/Plugins/Build/HordeServer.Build/Jobs/Bisect/BisectController.cs
+++ b/Engine/Source/Programs/Horde/Plugins/Build/HordeServer.Build/Jobs/Bisect/BisectController.cs
@@ -266,13 +266,14 @@ namespace HordeServer.Jobs.Bisect
 		readonly IGraphCollection _graphCollection;
 		readonly IUserCollection _userCollection;
 		readonly IOptionsSnapshot<BuildConfig> _buildConfig;
+		readonly BuildServerConfig _staticConfig;
 		private readonly ILogger<BisectTasksController> _logger;
 		readonly Tracer _tracer;
 
 		/// <summary>
 		/// Constructor
 		/// </summary>
-		public BisectTasksController(IBisectTaskCollection bisectTaskCollection, JobService jobService, IJobCollection jobCollection, IJobStepRefCollection jobStepRefs, IGraphCollection graphCollection, IUserCollection userCollection, Tracer tracer, ILogger<BisectTasksController> logger, IOptionsSnapshot<BuildConfig> buildConfig)
+		public BisectTasksController(IBisectTaskCollection bisectTaskCollection, JobService jobService, IJobCollection jobCollection, IJobStepRefCollection jobStepRefs, IGraphCollection graphCollection, IUserCollection userCollection, Tracer tracer, ILogger<BisectTasksController> logger, IOptionsSnapshot<BuildConfig> buildConfig, BuildServerConfig staticConfig)
 		{
 			_bisectTaskCollection = bisectTaskCollection;
 			_jobService = jobService;
@@ -281,6 +282,7 @@ namespace HordeServer.Jobs.Bisect
 			_graphCollection = graphCollection;
 			_userCollection = userCollection;
 			_buildConfig = buildConfig;
+			_staticConfig = staticConfig;
 			_tracer = tracer;
 			_logger = logger;
 		}
@@ -295,6 +297,11 @@ namespace HordeServer.Jobs.Bisect
 		[Route("/api/v1/bisect")]
 		public async Task<ActionResult<CreateBisectTaskResponse>> CreateAsync([FromBody] CreateBisectTaskRequest create, CancellationToken cancellationToken = default)
 		{
+			if (_staticConfig.DisableBisection)
+			{
+				return BadRequest("Bisection is currently disabled");
+			}
+
 			IJob? job = await _jobCollection.GetAsync(create.JobId, cancellationToken);
 			if (job == null)
 			{
@@ -379,6 +386,11 @@ namespace HordeServer.Jobs.Bisect
 		[Route("/api/v1/bisect/{bisectTaskId}")]
 		public async Task<ActionResult<GetBisectTaskResponse>> GetAsync([FromRoute] BisectTaskId bisectTaskId, CancellationToken cancellationToken = default)
 		{
+			if (_staticConfig.DisableBisection)
+			{
+				return BadRequest("Bisection is currently disabled");
+			}
+
 			IBisectTask? bisectTask = await _bisectTaskCollection.GetAsync(bisectTaskId, cancellationToken);
 			if (bisectTask == null)
 			{
@@ -409,6 +421,11 @@ namespace HordeServer.Jobs.Bisect
 		[Route("/api/v1/bisect/{bisectTaskId}")]
 		public async Task<ActionResult> UpdateAsync([FromRoute] BisectTaskId bisectTaskId, [FromBody] UpdateBisectTaskRequest request, CancellationToken cancellationToken)
 		{
+			if (_staticConfig.DisableBisection)
+			{
+				return BadRequest("Bisection is currently disabled");
+			}
+
 			for (; ; )
 			{
 				IBisectTask? bisectTask = await _bisectTaskCollection.GetAsync(bisectTaskId, cancellationToken);
@@ -483,6 +500,11 @@ namespace HordeServer.Jobs.Bisect
 		[Route("/api/v1/bisect/job/{jobId}")]
 		public async Task<ActionResult<List<GetBisectTaskResponse>>> GetJobBisectTasksAsync([FromRoute] JobId jobId, CancellationToken cancellationToken = default)
 		{
+			if (_staticConfig.DisableBisection)
+			{
+				return new List<GetBisectTaskResponse>();
+			}
+
 			IReadOnlyList<IBisectTask> tasks = await _bisectTaskCollection.FindAsync(null, jobId, null, null, null, null, null, cancellationToken);
 
 			List<GetBisectTaskResponse> response = new List<GetBisectTaskResponse>();
@@ -522,6 +544,11 @@ namespace HordeServer.Jobs.Bisect
 			[FromQuery] int count = 100,
 			CancellationToken cancellationToken = default)
 		{
+			if (_staticConfig.DisableBisection)
+			{
+				return new List<GetBisectTaskResponse>();
+			}
+
 			using TelemetrySpan span = _tracer.StartActiveSpan($"{nameof(BisectTasksController)}.{nameof(FindBisectTasksAsync)}");
 
 			List<GetBisectTaskResponse> responses = new List<GetBisectTaskResponse>();
diff --git a/Engine/Source/Programs/Horde/Plugins/Build/HordeServer.Build/Jobs/Bisect/BisectService.cs b/Engine/Source/Programs/Horde/Plugins/Build/HordeServer.Build/Jobs/Bisect/BisectService.cs
--- a/Engine/Source/Programs/Horde/Plugins/Build/HordeServer.Build/Jobs/Bisect/BisectService.cs
+++ b/Engine/Source/Programs/Horde/Plugins/Build/HordeServer.Build/Jobs/Bisect/BisectService.cs
@@ -25,10 +25,11 @@ namespace HordeServer.Jobs.Bisect
 		readonly IGraphCollection _graphCollection;
 		readonly ICommitService _commitService;
 		readonly IOptionsMonitor<BuildConfig> _buildConfig;
+		readonly BuildServerConfig _staticConfig;
 		readonly ILogger _logger;
 		readonly ITicker _ticker;
 
-		public BisectService(IBisectTaskCollection bisectTaskCollection, IJobCollection jobCollection, IJobStepRefCollection jobStepRefCollection, ITemplateCollectionInternal templateCollection, IGraphCollection graphCollection, ICommitService commitService, IOptionsMonitor<BuildConfig> buildConfig, IClock clock, ILogger<BisectService> logger)
+		public BisectService(IBisectTaskCollection bisectTaskCollection, IJobCollection jobCollection, IJobStepRefCollection jobStepRefCollection, ITemplateCollectionInternal templateCollection, IGraphCollection graphCollection, ICommitService commitService, IOptionsMonitor<BuildConfig> buildConfig, BuildServerConfig staticConfig, IClock clock, ILogger<BisectService> logger)
 		{
 			_bisectTaskCollection = bisectTaskCollection;
 			_jobCollection = jobCollection;
@@ -37,6 +38,7 @@ namespace HordeServer.Jobs.Bisect
 			_graphCollection = graphCollection;
 			_commitService = commitService;
 			_buildConfig = buildConfig;
+			_staticConfig = staticConfig;
 			_logger = logger;
 			_ticker = clock.AddSharedTicker<BisectService>(TimeSpan.FromMinutes(1.0), TickAsync, logger);
 		}
@@ -55,6 +57,11 @@ namespace HordeServer.Jobs.Bisect
 
 		async ValueTask TickAsync(CancellationToken cancellationToken)
 		{
+			if (_staticConfig.DisableBisection)
+			{
+				return;
+			}
+
 			BuildConfig buildConfig = _buildConfig.CurrentValue;
 			await foreach (IBisectTask bisectTask in _bisectTaskCollection.FindActiveAsync(cancellationToken))
 			{
