Feature flag bisection test - disabled by default
#rnx

diff --git a/Engine/Source/Programs/Horde/Plugins/Build/HordeServer.Build.Tests/Jobs/Bisect/BisectTests.cs b/Engine/Source/Programs/Horde/Plugins/Build/HordeServer.Build.Tests/Jobs/Bisect/BisectTests.cs
--- a/Engine/Source/Programs/Horde/Plugins/Build/HordeServer.Build.Tests/Jobs/Bisect/BisectTests.cs
+++ b/Engine/Source/Programs/Horde/Plugins/Build/HordeServer.Build.Tests/Jobs/Bisect/BisectTests.cs
@@ -26,6 +26,7 @@ namespace HordeServer.Tests.Jobs.Bisect
 	public class BisectTests : BuildTestSetup
 	{
 		[TestMethod]
+		[Ignore("Bisection is disabled via feature flag (DisableBisection)")]
 		public async Task TestStatesAsync()
 		{
 			(IJob failedJob, IGraph graph) = await SetupBisectionTestAsync();
