Horde: Update release notes

#rnx

diff --git a/Engine/Source/Programs/Horde/Docs/ReleaseNotes.md b/Engine/Source/Programs/Horde/Docs/ReleaseNotes.md
--- a/Engine/Source/Programs/Horde/Docs/ReleaseNotes.md
+++ b/Engine/Source/Programs/Horde/Docs/ReleaseNotes.md
@@ -1,5 +1,38 @@
 # Release Notes
 
+## 2026-03-12
+
+* Feature flag bisection - disabled by default (51640744)
+* EpicGames.OIDC: Fix concurrent refresh token handling
+  Fix a race condition where multiple concurrent processes (UBT, UGS, HordeCmd) could race on OIDC refresh token rotation, causing invalid_grant errors when one process uses a token already rotated by another.
+  * Always read the refresh token fresh from the backing store rather than caching it in memory
+  * Add a machine-wide file lock around the refresh operation so only one process refreshes at a time
+  * Retry with a re-read token on invalid_grant as a fallback for older binaries without the lock
+  * Add tests for refresh retry logic and WindowsTokenStore DPAPI round-trip (51629776)
+* Add ResolveSecret to TokenConfig.ClientId and ClientSecret (51622248)
+* Fix error 'Change has shelved files associated with it and can't be deleted' during conforms (51622061)
+* Handle missing log blobs gracefully in step history events
+  When log blob data has been garbage collected from storage, the
+  GetEventsAsync and GetBatchEventsAsync endpoints now catch
+  StorageException and skip the unavailable event instead of
+  returning a 500 error to the dashboard. (51592668)
+* Make tool deployment retention limit configurable
+  Adds a MaxDeployments property to ToolConfig (defaulting to 5, preserving existing behavior) to allow per-tool configuration of how many deployments are retained. (51589319)
+* Set a default telemetry store for compute (51549153)
+* Increase MongoDB batch size for agent reads to reduce getMore errors
+  Set a batch size of 10000 when fetching agent documents from MongoDB to minimize getMore commands on secondary reads. The agents collection is updated frequently, making multi-batch cursor reads prone to errors when reading from replicas. (51504525)
+* Fix build warnings in Storage and Build plugins
+  Add missing XML doc comments to CollectionStatsService and GcMetricsCollection.
+  Remove unused using directive in CollectionStatsService.
+  Make BuildStatusResponse static in StorageController (CA1822).
+  Suppress CA1505 on CreateOrUpdateWorkflowThreadAsync in SlackNotificationSink. (51465969)
+* Only register TunnelService when tunnel port is configured
+  Also make StopAsync log exceptions as warnings. (51465212)
+* Fix Analyze build configuration having no effect when building via .csproj
+  * .sln-based builds set the platform as "Any CPU" but direct .csproj builds do not unless you explicitly set this
+  * This causes certain C# builds to fail reporting warnings from analyze configuration
+  * Fix by conditioning only on Configuration, since Platform is irrelevant for these settings. This matches how the property is consumed in CI and locally. (51464192)
+
 ## 2026-03-05
 
 * EpicGames.LogMatchers: Fix a static analysis warning by specifying ordinal string comparison when checking if a log line contains the callstack indicator. (51456221)
