Always register BisectService for DI even when bisection is feature-flagged off

BisectService already checks DisableBisection at runtime in TickAsync,
so it is safe to register unconditionally. This ensures the controller
and related services remain available for dependency injection.

diff --git a/Engine/Source/Programs/Horde/Plugins/Build/HordeServer.Build/BuildPlugin.cs b/Engine/Source/Programs/Horde/Plugins/Build/HordeServer.Build/BuildPlugin.cs
--- a/Engine/Source/Programs/Horde/Plugins/Build/HordeServer.Build/BuildPlugin.cs
+++ b/Engine/Source/Programs/Horde/Plugins/Build/HordeServer.Build/BuildPlugin.cs
@@ -203,10 +203,7 @@ namespace HordeServer
 			if (_serverInfo.IsRunModeActive(RunMode.Worker) && !_serverInfo.ReadOnlyMode)
 			{
 				services.AddHostedService<AgentReportService>();
-				if (!_staticConfig.DisableBisection)
-				{
-					services.AddHostedService<BisectService>();
-				}
+				services.AddHostedService<BisectService>();
 				services.AddHostedService(provider => provider.GetRequiredService<IssueService>());
 				services.AddHostedService<IssueReportService>();
 				services.AddHostedService<IssueTagService>();
diff --git a/Engine/Source/Programs/Horde/Plugins/Build/HordeServer.Build/Jobs/Bisect/BisectController.cs b/Engine/Source/Programs/Horde/Plugins/Build/HordeServer.Build/Jobs/Bisect/BisectController.cs
--- a/Engine/Source/Programs/Horde/Plugins/Build/HordeServer.Build/Jobs/Bisect/BisectController.cs
+++ b/Engine/Source/Programs/Horde/Plugins/Build/HordeServer.Build/Jobs/Bisect/BisectController.cs
@@ -266,14 +266,14 @@ namespace HordeServer.Jobs.Bisect
 		readonly IGraphCollection _graphCollection;
 		readonly IUserCollection _userCollection;
 		readonly IOptionsSnapshot<BuildConfig> _buildConfig;
-		readonly BuildServerConfig _staticConfig;
+		readonly IOptions<BuildServerConfig> _staticConfig;
 		private readonly ILogger<BisectTasksController> _logger;
 		readonly Tracer _tracer;
 
 		/// <summary>
 		/// Constructor
 		/// </summary>
-		public BisectTasksController(IBisectTaskCollection bisectTaskCollection, JobService jobService, IJobCollection jobCollection, IJobStepRefCollection jobStepRefs, IGraphCollection graphCollection, IUserCollection userCollection, Tracer tracer, ILogger<BisectTasksController> logger, IOptionsSnapshot<BuildConfig> buildConfig, BuildServerConfig staticConfig)
+		public BisectTasksController(IBisectTaskCollection bisectTaskCollection, JobService jobService, IJobCollection jobCollection, IJobStepRefCollection jobStepRefs, IGraphCollection graphCollection, IUserCollection userCollection, Tracer tracer, ILogger<BisectTasksController> logger, IOptionsSnapshot<BuildConfig> buildConfig, IOptions<BuildServerConfig> staticConfig)
 		{
 			_bisectTaskCollection = bisectTaskCollection;
 			_jobService = jobService;
@@ -297,7 +297,7 @@ namespace HordeServer.Jobs.Bisect
 		[Route("/api/v1/bisect")]
 		public async Task<ActionResult<CreateBisectTaskResponse>> CreateAsync([FromBody] CreateBisectTaskRequest create, CancellationToken cancellationToken = default)
 		{
-			if (_staticConfig.DisableBisection)
+			if (_staticConfig.Value.DisableBisection)
 			{
 				return BadRequest("Bisection is currently disabled");
 			}
@@ -386,7 +386,7 @@ namespace HordeServer.Jobs.Bisect
 		[Route("/api/v1/bisect/{bisectTaskId}")]
 		public async Task<ActionResult<GetBisectTaskResponse>> GetAsync([FromRoute] BisectTaskId bisectTaskId, CancellationToken cancellationToken = default)
 		{
-			if (_staticConfig.DisableBisection)
+			if (_staticConfig.Value.DisableBisection)
 			{
 				return BadRequest("Bisection is currently disabled");
 			}
@@ -421,7 +421,7 @@ namespace HordeServer.Jobs.Bisect
 		[Route("/api/v1/bisect/{bisectTaskId}")]
 		public async Task<ActionResult> UpdateAsync([FromRoute] BisectTaskId bisectTaskId, [FromBody] UpdateBisectTaskRequest request, CancellationToken cancellationToken)
 		{
-			if (_staticConfig.DisableBisection)
+			if (_staticConfig.Value.DisableBisection)
 			{
 				return BadRequest("Bisection is currently disabled");
 			}
@@ -500,7 +500,7 @@ namespace HordeServer.Jobs.Bisect
 		[Route("/api/v1/bisect/job/{jobId}")]
 		public async Task<ActionResult<List<GetBisectTaskResponse>>> GetJobBisectTasksAsync([FromRoute] JobId jobId, CancellationToken cancellationToken = default)
 		{
-			if (_staticConfig.DisableBisection)
+			if (_staticConfig.Value.DisableBisection)
 			{
 				return new List<GetBisectTaskResponse>();
 			}
@@ -544,7 +544,7 @@ namespace HordeServer.Jobs.Bisect
 			[FromQuery] int count = 100,
 			CancellationToken cancellationToken = default)
 		{
-			if (_staticConfig.DisableBisection)
+			if (_staticConfig.Value.DisableBisection)
 			{
 				return new List<GetBisectTaskResponse>();
 			}
diff --git a/Engine/Source/Programs/Horde/Plugins/Build/HordeServer.Build/Jobs/Bisect/BisectService.cs b/Engine/Source/Programs/Horde/Plugins/Build/HordeServer.Build/Jobs/Bisect/BisectService.cs
--- a/Engine/Source/Programs/Horde/Plugins/Build/HordeServer.Build/Jobs/Bisect/BisectService.cs
+++ b/Engine/Source/Programs/Horde/Plugins/Build/HordeServer.Build/Jobs/Bisect/BisectService.cs
@@ -25,11 +25,11 @@ namespace HordeServer.Jobs.Bisect
 		readonly IGraphCollection _graphCollection;
 		readonly ICommitService _commitService;
 		readonly IOptionsMonitor<BuildConfig> _buildConfig;
-		readonly BuildServerConfig _staticConfig;
+		readonly IOptions<BuildServerConfig> _staticConfig;
 		readonly ILogger _logger;
 		readonly ITicker _ticker;
 
-		public BisectService(IBisectTaskCollection bisectTaskCollection, IJobCollection jobCollection, IJobStepRefCollection jobStepRefCollection, ITemplateCollectionInternal templateCollection, IGraphCollection graphCollection, ICommitService commitService, IOptionsMonitor<BuildConfig> buildConfig, BuildServerConfig staticConfig, IClock clock, ILogger<BisectService> logger)
+		public BisectService(IBisectTaskCollection bisectTaskCollection, IJobCollection jobCollection, IJobStepRefCollection jobStepRefCollection, ITemplateCollectionInternal templateCollection, IGraphCollection graphCollection, ICommitService commitService, IOptionsMonitor<BuildConfig> buildConfig, IOptions<BuildServerConfig> staticConfig, IClock clock, ILogger<BisectService> logger)
 		{
 			_bisectTaskCollection = bisectTaskCollection;
 			_jobCollection = jobCollection;
@@ -57,7 +57,7 @@ namespace HordeServer.Jobs.Bisect
 
 		async ValueTask TickAsync(CancellationToken cancellationToken)
 		{
-			if (_staticConfig.DisableBisection)
+			if (_staticConfig.Value.DisableBisection)
 			{
 				return;
 			}
