Horde: Swap registration order of exception logger vs handler

Surfaces the error messages in ExceptionLoggingMiddleware, which previously got silenced. This so we can control the format of "Unhandled exception" log lines. Once verified, we'll disable the generic ASP.NET based logger.

diff --git a/Engine/Source/Programs/Horde/HordeServer/Server/ExceptionController.cs b/Engine/Source/Programs/Horde/HordeServer/Server/ExceptionController.cs
--- a/Engine/Source/Programs/Horde/HordeServer/Server/ExceptionController.cs
+++ b/Engine/Source/Programs/Horde/HordeServer/Server/ExceptionController.cs
@@ -37,7 +37,7 @@ namespace HordeServer.Server
 			}
 			catch (Exception ex)
 			{
-				logger.LogError(ex, "Unhandled exception for {Method} {Path} {TraceId}", context.Request.Method, context.Request.Path, context.TraceIdentifier);
+				logger.LogError(ex, "Unhandled exception for {Method} {Path}: {ExceptionMessage}", context.Request.Method, context.Request.Path, ex.Message);
 				context.Items[LoggedContextKey] = true;
 				throw;
 			}
