Horde: Simplify session conflict handling to use debug logging + batched Slack report

Replace the unused auto-quarantine implementation with a lightweight Redis sorted set: the first 20 mismatches per agent log at error level subsequent ones log at debug.

diff --git a/Engine/Source/Programs/Horde/Plugins/Build/HordeServer.Build.Tests/Notifications/NotificationServiceTests.cs b/Engine/Source/Programs/Horde/Plugins/Build/HordeServer.Build.Tests/Notifications/NotificationServiceTests.cs
--- a/Engine/Source/Programs/Horde/Plugins/Build/HordeServer.Build.Tests/Notifications/NotificationServiceTests.cs
+++ b/Engine/Source/Programs/Horde/Plugins/Build/HordeServer.Build.Tests/Notifications/NotificationServiceTests.cs
@@ -93,6 +93,7 @@ namespace HordeServer.Tests.Notifications
 		public Task NotifyConfigUpdateFailureAsync(string errorMessage, string fileName, int? change = null, IUser? author = null, string? description = null, CancellationToken cancellationToken = default) { throw new NotImplementedException(); }
 		public Task NotifyDeviceServiceAsync(string message, IDevice? device = null, IDevicePool? pool = null, StreamConfig? stream = null, IJob? job = null, IJobStep? step = null, INode? node = null, IUser? user = null, CancellationToken cancellationToken = default) { throw new NotImplementedException(); }
 		public Task SendAgentReportAsync(AgentReport report, CancellationToken cancellationToken) => throw new NotImplementedException();
+		public Task SendSessionConflictReportAsync(IReadOnlyList<(AgentId Id, int Count)> conflicts, CancellationToken cancellationToken) => Task.CompletedTask;
 		public Task SendIssueReportAsync(IssueReportGroup report, CancellationToken cancellationToken) => throw new NotImplementedException();
 		public Task SendDeviceIssueReportAsync(DeviceIssueReport report, CancellationToken cancellationToken) => throw new NotImplementedException();
 		public Task NotifyTestHealthReportAsync(ITestHealthReport report, string recipient, string[]? carbonCopies, CancellationToken cancellationToken) => throw new NotImplementedException();
diff --git a/Engine/Source/Programs/Horde/Plugins/Build/HordeServer.Build.Tests/Server/RpcServiceTest.cs b/Engine/Source/Programs/Horde/Plugins/Build/HordeServer.Build.Tests/Server/RpcServiceTest.cs
--- a/Engine/Source/Programs/Horde/Plugins/Build/HordeServer.Build.Tests/Server/RpcServiceTest.cs
+++ b/Engine/Source/Programs/Horde/Plugins/Build/HordeServer.Build.Tests/Server/RpcServiceTest.cs
@@ -2,7 +2,6 @@
 
 using System.IdentityModel.Tokens.Jwt;
 using System.IO.Compression;
-using System.Net;
 using System.Security.Claims;
 using System.Text;
 using EpicGames.Horde.Agents;
@@ -16,7 +15,6 @@ using HordeServer.Agents;
 using HordeServer.Tools;
 using HordeServer.Utilities;
 using Microsoft.Extensions.DependencyInjection;
-using Microsoft.Extensions.Options;
 
 namespace HordeServer.Tests.Server
 {
@@ -189,6 +187,64 @@ namespace HordeServer.Tests.Server
 			Assert.IsTrue(re.Message.Contains("Invalid agent name", StringComparison.OrdinalIgnoreCase));
 		}
 
+		[TestMethod]
+		public async Task SessionMismatchWithLowCountThrowsStructuredRpcExceptionAsync()
+		{
+			// Create first session
+			RpcCreateSessionRequest createReq1 = new() { Id = new AgentId("MismatchAgent").ToString(), Capabilities = new RpcAgentCapabilities() };
+			RpcCreateSessionResponse createRes1 = await RpcService.CreateSession(createReq1, _adminContext);
+			string staleSessionId = createRes1.SessionId;
+
+			// Create a second session (invalidates first)
+			RpcCreateSessionRequest createReq2 = new() { Id = new AgentId("MismatchAgent").ToString(), Capabilities = new RpcAgentCapabilities() };
+			await RpcService.CreateSession(createReq2, _adminContext);
+
+			// Try to update using stale session → should throw StructuredRpcException (count ≤ threshold)
+			ServerCallContext sessionContext = ServerCallContextStub.ForAdminWithAgentSessionId(staleSessionId);
+			TestAsyncStreamReader<RpcUpdateSessionRequest> requestStream = new(sessionContext);
+			TestServerStreamWriter<RpcUpdateSessionResponse> responseStream = new(sessionContext);
+			Task call = RpcService.UpdateSession(requestStream, responseStream, sessionContext);
+
+			requestStream.AddMessage(new RpcUpdateSessionRequest { AgentId = createRes1.AgentId, SessionId = staleSessionId });
+			StructuredRpcException ex = await Assert.ThrowsExactlyAsync<StructuredRpcException>(() => call);
+			Assert.AreEqual(StatusCode.PermissionDenied, ex.StatusCode);
+			Assert.IsTrue(ex.Message.Contains("has completed session", StringComparison.OrdinalIgnoreCase));
+		}
+
+		[TestMethod]
+		public async Task SessionMismatchWithHighCountThrowsPlainRpcExceptionAsync()
+		{
+			// Create first session
+			RpcCreateSessionRequest createReq1 = new() { Id = new AgentId("MismatchAgent2").ToString(), Capabilities = new RpcAgentCapabilities() };
+			RpcCreateSessionResponse createRes1 = await RpcService.CreateSession(createReq1, _adminContext);
+			string staleSessionId = createRes1.SessionId;
+			AgentId agentId = new(createRes1.AgentId);
+
+			// Create a second session (invalidates first)
+			RpcCreateSessionRequest createReq2 = new() { Id = new AgentId("MismatchAgent2").ToString(), Capabilities = new RpcAgentCapabilities() };
+			await RpcService.CreateSession(createReq2, _adminContext);
+
+			// Pre-populate conflict count up to MaxDetailedConflictErrors via SessionConflictService
+			SessionConflictService conflictService = ServiceProvider.GetRequiredService<SessionConflictService>();
+			for (int i = 0; i < HordeServer.Server.RpcService.MaxDetailedConflictErrors; i++)
+			{
+				await conflictService.RecordConflictAsync(agentId);
+			}
+
+			// Try to update using stale session → should throw plain RpcException (count > threshold)
+			ServerCallContext sessionContext = ServerCallContextStub.ForAdminWithAgentSessionId(staleSessionId);
+			TestAsyncStreamReader<RpcUpdateSessionRequest> requestStream = new(sessionContext);
+			TestServerStreamWriter<RpcUpdateSessionResponse> responseStream = new(sessionContext);
+			Task call = RpcService.UpdateSession(requestStream, responseStream, sessionContext);
+
+			requestStream.AddMessage(new RpcUpdateSessionRequest { AgentId = createRes1.AgentId, SessionId = staleSessionId });
+			RpcException ex = await Assert.ThrowsExactlyAsync<RpcException>(() => call);
+			Assert.AreEqual(StatusCode.PermissionDenied, ex.StatusCode);
+
+			// Should NOT be a StructuredRpcException
+			Assert.IsFalse(ex is StructuredRpcException);
+		}
+
 		[TestMethod]
 		public async Task FinishBatchTestAsync()
 		{
@@ -249,55 +305,6 @@ namespace HordeServer.Tests.Server
 			}
 		}
 		
-		[TestMethod]
-		public async Task UpdateSession_SessionMismatch_QuarantinesAfterThresholdAsync()
-		{
-			ComputeServerConfig computeServerConfig = ServiceProvider.GetRequiredService<IOptionsMonitor<ComputeServerConfig>>().CurrentValue;
-			computeServerConfig.SessionConflictEnabled = true;
-			int threshold = computeServerConfig.SessionConflictThreshold;
-			AgentId agentId = new AgentId("conflict-agent");
-			IPAddress callerIp = IPAddress.Parse("10.0.0.1");
-			ServerCallContext adminCtx = ServerCallContextStub.ForAdmin(callerIp);
-
-			// Create agent and first session
-			RpcCreateSessionRequest createReq1 = new RpcCreateSessionRequest
-			{
-				Id = agentId.ToString(),
-				Capabilities = new RpcAgentCapabilities()
-			};
-			RpcCreateSessionResponse createRes1 = await RpcService.CreateSession(createReq1, adminCtx);
-			string staleSessionId = createRes1.SessionId;
-
-			// Create a second session (invalidates the first)
-			RpcCreateSessionRequest createReq2 = new RpcCreateSessionRequest
-			{
-				Id = agentId.ToString(),
-				Capabilities = new RpcAgentCapabilities()
-			};
-			await RpcService.CreateSession(createReq2, adminCtx);
-
-			// Attempt UpdateSession with the stale session ID repeatedly
-			// Each call should hit the session mismatch path and record a conflict
-			for (int i = 0; i < threshold; i++)
-			{
-				TestAsyncStreamReader<RpcUpdateSessionRequest> requestStream = new(adminCtx);
-				TestServerStreamWriter<RpcUpdateSessionResponse> responseStream = new(adminCtx);
-				requestStream.AddMessage(new RpcUpdateSessionRequest { AgentId = agentId.ToString(), SessionId = staleSessionId });
-
-				StructuredRpcException ex = await Assert.ThrowsExactlyAsync<StructuredRpcException>(
-					() => RpcService.UpdateSession(requestStream, responseStream, adminCtx));
-				Assert.AreEqual(StatusCode.PermissionDenied, ex.StatusCode);
-			}
-
-			// Verify the agent has been quarantined
-			IAgent? agent = await AgentService.GetAgentAsync(agentId);
-			Assert.IsNotNull(agent);
-			Assert.IsFalse(agent.Enabled, "Agent should be disabled after session conflict threshold");
-			Assert.IsNotNull(agent.Comment);
-			Assert.IsTrue(agent.Comment.StartsWith(SessionConflictService.CommentPrefix, StringComparison.Ordinal),
-				$"Comment should start with session-conflict prefix, was: {agent.Comment}");
-		}
-
 		private static ServerCallContextStub CreateServerContext(string? roleClaim)
 		{
 			return new ServerCallContextStub(roleClaim != null ? new AclClaimConfig(HordeClaimTypes.Role, roleClaim).ToClaim() : null);
diff --git a/Engine/Source/Programs/Horde/Plugins/Build/HordeServer.Build/Agents/AgentReportService.cs b/Engine/Source/Programs/Horde/Plugins/Build/HordeServer.Build/Agents/AgentReportService.cs
--- a/Engine/Source/Programs/Horde/Plugins/Build/HordeServer.Build/Agents/AgentReportService.cs
+++ b/Engine/Source/Programs/Horde/Plugins/Build/HordeServer.Build/Agents/AgentReportService.cs
@@ -25,11 +25,6 @@ namespace HordeServer.Agents
 		/// List of agents stuck in a upgrade loop, plus the number of times they've attempted to upgrade
 		/// </summary>
 		public List<(AgentId, int)> UpgradeLoop { get; } = new List<(AgentId, int)>();
-
-		/// <summary>
-		/// List of agents quarantined due to session conflicts, with their comment text
-		/// </summary>
-		public List<(AgentId Id, string Comment)> SessionConflicts { get; } = new List<(AgentId, string)>();
 	}
 
 	[SingletonDocument("agent-report-state", "6268871c211d05611b3e4fd8")]
@@ -104,11 +99,6 @@ namespace HordeServer.Agents
 							report.UpgradeLoop.Add((agent.Id, agent.UpgradeAttemptCount.Value));
 						}
 					}
-
-					if (!agent.Enabled && agent.Comment?.StartsWith(SessionConflictService.CommentPrefix, StringComparison.Ordinal) == true)
-					{
-						report.SessionConflicts.Add((agent.Id, agent.Comment));
-					}
 				}
 
 				await _notificationService.SendAgentReportAsync(report, cancellationToken);
diff --git a/Engine/Source/Programs/Horde/Plugins/Build/HordeServer.Build/Notifications/INotificationSink.cs b/Engine/Source/Programs/Horde/Plugins/Build/HordeServer.Build/Notifications/INotificationSink.cs
--- a/Engine/Source/Programs/Horde/Plugins/Build/HordeServer.Build/Notifications/INotificationSink.cs
+++ b/Engine/Source/Programs/Horde/Plugins/Build/HordeServer.Build/Notifications/INotificationSink.cs
@@ -1,5 +1,6 @@
 // Copyright Epic Games, Inc. All Rights Reserved.
 
+using EpicGames.Horde.Agents;
 using EpicGames.Horde.Jobs;
 using EpicGames.Horde.Jobs.Graphs;
 using EpicGames.Horde.Users;
@@ -156,6 +157,13 @@ namespace HordeServer.Notifications
 		/// </summary>
 		Task SendAgentReportAsync(AgentReport report, CancellationToken cancellationToken);
 
+		/// <summary>
+		/// Send a batched report of session conflicts detected in the last 12 hours
+		/// </summary>
+		/// <param name="conflicts">Agent IDs paired with their total mismatch count</param>
+		/// <param name="cancellationToken">Cancellation token for the operation</param>
+		Task SendSessionConflictReportAsync(IReadOnlyList<(AgentId Id, int Count)> conflicts, CancellationToken cancellationToken);
+
 		/// <summary>
 		/// Post a notification for test health report
 		/// </summary>
diff --git a/Engine/Source/Programs/Horde/Plugins/Build/HordeServer.Build/Notifications/NotificationService.cs b/Engine/Source/Programs/Horde/Plugins/Build/HordeServer.Build/Notifications/NotificationService.cs
--- a/Engine/Source/Programs/Horde/Plugins/Build/HordeServer.Build/Notifications/NotificationService.cs
+++ b/Engine/Source/Programs/Horde/Plugins/Build/HordeServer.Build/Notifications/NotificationService.cs
@@ -3,6 +3,7 @@
 using System.Diagnostics.Metrics;
 using System.Security.Claims;
 using EpicGames.Core;
+using EpicGames.Horde.Agents;
 using EpicGames.Horde.Jobs;
 using EpicGames.Horde.Jobs.Graphs;
 using EpicGames.Horde.Notifications;
@@ -118,6 +119,7 @@ namespace HordeServer.Notifications
 		internal TimeSpan _notificationBatchInterval = TimeSpan.FromHours(12);
 
 		readonly IServerInfo _serverInfo;
+		readonly SessionConflictService _sessionConflictService;
 		readonly Counter<int> _jobCounter;
 		readonly Histogram<double> _jobDurationHistogram;
 
@@ -142,6 +144,7 @@ namespace HordeServer.Notifications
 			IRedisService redisService,
 			IConfigService configService,
 			IServerInfo serverInfo,
+			SessionConflictService sessionConflictService,
 			IClock clock)
 		{
 			_sinks = sinks.ToList();
@@ -157,6 +160,7 @@ namespace HordeServer.Notifications
 			_cache = cache;
 			_redisConnectionPool = redisService.ConnectionPool;
 			_serverInfo = serverInfo;
+			_sessionConflictService = sessionConflictService;
 
 			issueService.OnIssueUpdated += NotifyIssueUpdated;
 			jobService.OnJobAborted += NotifyJobAborted;
@@ -459,6 +463,30 @@ namespace HordeServer.Notifications
 					await sink.NotifyJobScheduledAsync(jobScheduledNotifications, cancellationToken);
 				}
 			}
+
+			try
+			{
+				List<(AgentId Id, int Count)> conflicts = await _sessionConflictService.DrainConflictsAsync(cancellationToken);
+				if (conflicts.Count > 0)
+				{
+					_logger.LogInformation("Drained {Count} agents with session conflicts", conflicts.Count);
+					foreach (INotificationSink sink in _sinks)
+					{
+						try
+						{
+							await sink.SendSessionConflictReportAsync(conflicts, cancellationToken);
+						}
+						catch (Exception e)
+						{
+							_logger.LogError(e, "Failed sending session conflict report");
+						}
+					}
+				}
+			}
+			catch (Exception e)
+			{
+				_logger.LogError(e, "Failed draining session conflicts from Redis");
+			}
 		}
 
 		/// <summary>
diff --git a/Engine/Source/Programs/Horde/Plugins/Build/HordeServer.Build/Notifications/Sinks/SlackNotificationSink.cs b/Engine/Source/Programs/Horde/Plugins/Build/HordeServer.Build/Notifications/Sinks/SlackNotificationSink.cs
--- a/Engine/Source/Programs/Horde/Plugins/Build/HordeServer.Build/Notifications/Sinks/SlackNotificationSink.cs
+++ b/Engine/Source/Programs/Horde/Plugins/Build/HordeServer.Build/Notifications/Sinks/SlackNotificationSink.cs
@@ -4015,6 +4015,31 @@ namespace HordeServer.Notifications.Sinks
 			}
 		}
 
+		/// <inheritdoc/>
+		public async Task SendSessionConflictReportAsync(IReadOnlyList<(AgentId Id, int Count)> conflicts, CancellationToken cancellationToken)
+		{
+			if (conflicts.Count == 0 || String.IsNullOrEmpty(_settings.AgentNotificationChannel))
+			{
+				return;
+			}
+
+			_logger.LogInformation("Sending session conflict report ({Count} agents) to {Channel}", conflicts.Count, _settings.AgentNotificationChannel);
+
+			const int NumItems = 10;
+			StringBuilder message = new StringBuilder("*Session conflicts (last 12h):*\n");
+			foreach ((AgentId agentId, int count) in conflicts.OrderByDescending(x => x.Count).ThenBy(x => x.Id.ToString()).Take(NumItems))
+			{
+				Uri agentUrl = new Uri(_serverInfo.DashboardUrl, $"agents?agentId={agentId}");
+				message.Append($"\u2022 *<{agentUrl}|{agentId}>* \u2014 {count:N0} mismatches\n");
+			}
+			if (conflicts.Count > NumItems)
+			{
+				message.Append($"+ {conflicts.Count - NumItems} other(s)\n");
+			}
+
+			await _slackClient.PostMessageAsync(_settings.AgentNotificationChannel, message.ToString(), cancellationToken);
+		}
+
 		const string DeepLinkUrl = "https://slack.com/app_redirect";
 		const int MaxDirectMessageUsers = 8;
 
diff --git a/Engine/Source/Programs/Horde/Plugins/Compute/HordeServer.Compute.Tests/Agents/SessionConflictServiceTest.cs b/Engine/Source/Programs/Horde/Plugins/Compute/HordeServer.Compute.Tests/Agents/SessionConflictServiceTest.cs
--- a/Engine/Source/Programs/Horde/Plugins/Compute/HordeServer.Compute.Tests/Agents/SessionConflictServiceTest.cs
+++ b/Engine/Source/Programs/Horde/Plugins/Compute/HordeServer.Compute.Tests/Agents/SessionConflictServiceTest.cs
@@ -1,136 +1,71 @@
 // Copyright Epic Games, Inc. All Rights Reserved.
 
-using System.Net;
 using EpicGames.Horde.Agents;
-using HordeServer.Agents.Sessions;
 using HordeServer.Agents;
-using HordeServer.Server;
 using Microsoft.Extensions.DependencyInjection;
-using Microsoft.Extensions.Logging;
-using Microsoft.Extensions.Options;
-using StackExchange.Redis;
 
-namespace HordeServer.Tests.Agents
-{
-	[TestClass]
-	public class SessionConflictServiceTest : ComputeTestSetup
-	{
-		private readonly IRedisService _redisService;
-
-		public SessionConflictServiceTest()
-		{
-			_redisService = GetRedisServiceSingleton();
-		}
-
-		private SessionConflictService CreateService(ComputeServerConfig? config = null)
-		{
-			IOptionsMonitor<ComputeServerConfig> configMonitor = config != null
-				? new TestOptionsMonitor<ComputeServerConfig>(config)
-				: ServiceProvider.GetRequiredService<IOptionsMonitor<ComputeServerConfig>>();
-			ILogger<SessionConflictService> logger = ServiceProvider.GetRequiredService<ILogger<SessionConflictService>>();
-			return new SessionConflictService(_redisService, AgentCollection, configMonitor, logger);
-		}
-
-		private static readonly IPAddress s_ip1 = IPAddress.Parse("10.1.2.3");
-		private static readonly IPAddress s_ip2 = IPAddress.Parse("10.1.2.4");
-
-		[TestMethod]
-		public async Task BelowThreshold_NoQuarantineAsync()
-		{
-			SessionConflictService service = CreateService(new ComputeServerConfig { SessionConflictEnabled = true });
-			IAgent agent = await AgentCollection.AddAsync(new CreateAgentOptions(new AgentId("conflict-test-1"), AgentMode.Dedicated, false, ""));
-
-			// Record 2 conflicts (threshold is 3)
-			await service.RecordConflictAsync(agent.Id, SessionIdUtils.GenerateNewId(), SessionIdUtils.GenerateNewId(), s_ip1, CancellationToken.None);
-			await service.RecordConflictAsync(agent.Id, SessionIdUtils.GenerateNewId(), SessionIdUtils.GenerateNewId(), s_ip2, CancellationToken.None);
-
-			IAgent? updated = await AgentCollection.GetAsync(agent.Id);
-			Assert.IsNotNull(updated);
-			Assert.IsTrue(updated.Enabled, "Agent should still be enabled below threshold");
-		}
-
-		[TestMethod]
-		public async Task AtThreshold_QuarantinesAgentAsync()
-		{
-			SessionConflictService service = CreateService(new ComputeServerConfig { SessionConflictEnabled = true });
-			IAgent agent = await AgentCollection.AddAsync(new CreateAgentOptions(new AgentId("conflict-test-2"), AgentMode.Dedicated, false, ""));
-
-			// Record 3 conflicts (hits threshold)
-			await service.RecordConflictAsync(agent.Id, SessionIdUtils.GenerateNewId(), SessionIdUtils.GenerateNewId(), s_ip1, CancellationToken.None);
-			await service.RecordConflictAsync(agent.Id, SessionIdUtils.GenerateNewId(), SessionIdUtils.GenerateNewId(), s_ip2, CancellationToken.None);
-			await service.RecordConflictAsync(agent.Id, SessionIdUtils.GenerateNewId(), SessionIdUtils.GenerateNewId(), s_ip1, CancellationToken.None);
-
-			IAgent? updated = await AgentCollection.GetAsync(agent.Id);
-			Assert.IsNotNull(updated);
-			Assert.IsFalse(updated.Enabled, "Agent should be disabled after reaching threshold");
-			Assert.IsNotNull(updated.Comment);
-			Assert.IsTrue(updated.Comment.StartsWith(SessionConflictService.CommentPrefix, StringComparison.Ordinal), $"Comment should start with prefix, was: {updated.Comment}");
-			Assert.IsTrue(updated.Comment.Contains("10.1.2.3", StringComparison.Ordinal), "Comment should contain IP 1");
-			Assert.IsTrue(updated.Comment.Contains("10.1.2.4", StringComparison.Ordinal), "Comment should contain IP 2");
-		}
-
-		[TestMethod]
-		public async Task WindowExpiry_OldEventsNotCountedAsync()
-		{
-			ComputeServerConfig config = new ComputeServerConfig { SessionConflictEnabled = true, SessionConflictWindowSeconds = 300 };
-			SessionConflictService service = CreateService(config);
-			IAgent agent = await AgentCollection.AddAsync(new CreateAgentOptions(new AgentId("conflict-test-3"), AgentMode.Dedicated, false, ""));
+namespace HordeServer.Tests.Agents;
 
-			// Insert 2 events with timestamps well outside the window (10 minutes old)
-			IDatabase db = _redisService.GetDatabase();
-			RedisKey key = SessionConflictService.GetRedisKey(agent.Id);
-			long expiredTicks = DateTimeOffset.UtcNow.AddSeconds(-600).Ticks;
-			await db.SortedSetAddAsync(key, $"old1:{Guid.NewGuid():N}", expiredTicks);
-			await db.SortedSetAddAsync(key, $"old2:{Guid.NewGuid():N}", expiredTicks);
-
-			// Record 1 more conflict via the service (triggers pruning of old entries)
-			await service.RecordConflictAsync(agent.Id, SessionIdUtils.GenerateNewId(), SessionIdUtils.GenerateNewId(), s_ip1, CancellationToken.None);
-
-			// Only 1 event should remain in the window (below threshold of 3)
-			long count = await db.SortedSetLengthAsync(key);
-			Assert.AreEqual(1, count, "Only the fresh event should remain after pruning");
+[TestClass]
+public class SessionConflictServiceTest : ComputeTestSetup
+{
+	SessionConflictService SessionConflictService => ServiceProvider.GetRequiredService<SessionConflictService>();
 
-			IAgent? updated = await AgentCollection.GetAsync(agent.Id);
-			Assert.IsNotNull(updated);
-			Assert.IsTrue(updated.Enabled, "Agent should still be enabled when old events were pruned");
-		}
+	[TestMethod]
+	public async Task RecordConflictReturnsIncrementingCountAsync()
+	{
+		AgentId agentId = new("agent-1");
+		Assert.AreEqual(1, await SessionConflictService.RecordConflictAsync(agentId));
+		Assert.AreEqual(2, await SessionConflictService.RecordConflictAsync(agentId));
+		Assert.AreEqual(3, await SessionConflictService.RecordConflictAsync(agentId));
+	}
 
-		[TestMethod]
-		public async Task AlreadyDisabled_IdempotentAsync()
-		{
-			SessionConflictService service = CreateService(new ComputeServerConfig { SessionConflictEnabled = true });
-			IAgent agent = await AgentCollection.AddAsync(new CreateAgentOptions(new AgentId("conflict-test-4"), AgentMode.Dedicated, false, ""));
+	[TestMethod]
+	public async Task MultipleAgentsTrackedIndependentlyAsync()
+	{
+		AgentId agent1 = new("agent-1");
+		AgentId agent2 = new("agent-2");
 
-			// Pre-disable the agent with a custom comment
-			agent = (await agent.TryUpdateAsync(new UpdateAgentOptions(Enabled: false, Comment: "manually disabled")))!;
-			Assert.IsFalse(agent.Enabled);
+		Assert.AreEqual(1, await SessionConflictService.RecordConflictAsync(agent1));
+		Assert.AreEqual(1, await SessionConflictService.RecordConflictAsync(agent2));
+		Assert.AreEqual(2, await SessionConflictService.RecordConflictAsync(agent1));
+		Assert.AreEqual(2, await SessionConflictService.RecordConflictAsync(agent2));
+	}
 
-			// Record 3+ conflicts (would normally trigger quarantine)
-			await service.RecordConflictAsync(agent.Id, SessionIdUtils.GenerateNewId(), SessionIdUtils.GenerateNewId(), s_ip1, CancellationToken.None);
-			await service.RecordConflictAsync(agent.Id, SessionIdUtils.GenerateNewId(), SessionIdUtils.GenerateNewId(), s_ip2, CancellationToken.None);
-			await service.RecordConflictAsync(agent.Id, SessionIdUtils.GenerateNewId(), SessionIdUtils.GenerateNewId(), s_ip1, CancellationToken.None);
+	[TestMethod]
+	public async Task DrainReturnsAllAgentsWithCountsAsync()
+	{
+		AgentId agent1 = new("agent-1");
+		AgentId agent2 = new("agent-2");
+
+		await SessionConflictService.RecordConflictAsync(agent1);
+		await SessionConflictService.RecordConflictAsync(agent1);
+		await SessionConflictService.RecordConflictAsync(agent1);
+		await SessionConflictService.RecordConflictAsync(agent2);
+
+		List<(AgentId Id, int Count)> conflicts = await SessionConflictService.DrainConflictsAsync();
+		Assert.AreEqual(2, conflicts.Count);
+		Assert.IsTrue(conflicts.Any(c => c.Id == agent1 && c.Count == 3));
+		Assert.IsTrue(conflicts.Any(c => c.Id == agent2 && c.Count == 1));
+	}
 
-			IAgent? updated = await AgentCollection.GetAsync(agent.Id);
-			Assert.IsNotNull(updated);
-			Assert.IsFalse(updated.Enabled);
-			Assert.AreEqual("manually disabled", updated.Comment, "Comment should be unchanged when agent was already disabled");
-		}
+	[TestMethod]
+	public async Task DrainClearsSortedSetAsync()
+	{
+		AgentId agentId = new("agent-1");
+		await SessionConflictService.RecordConflictAsync(agentId);
 
-		[TestMethod]
-		public async Task Disabled_NoRecordingAsync()
-		{
-			ComputeServerConfig disabledConfig = new ComputeServerConfig { SessionConflictEnabled = false };
-			SessionConflictService service = CreateService(disabledConfig);
-			IAgent agent = await AgentCollection.AddAsync(new CreateAgentOptions(new AgentId("conflict-test-5"), AgentMode.Dedicated, false, ""));
+		List<(AgentId Id, int Count)> first = await SessionConflictService.DrainConflictsAsync();
+		Assert.AreEqual(1, first.Count);
 
-			// Record 3+ conflicts with feature disabled
-			await service.RecordConflictAsync(agent.Id, SessionIdUtils.GenerateNewId(), SessionIdUtils.GenerateNewId(), s_ip1, CancellationToken.None);
-			await service.RecordConflictAsync(agent.Id, SessionIdUtils.GenerateNewId(), SessionIdUtils.GenerateNewId(), s_ip2, CancellationToken.None);
-			await service.RecordConflictAsync(agent.Id, SessionIdUtils.GenerateNewId(), SessionIdUtils.GenerateNewId(), s_ip1, CancellationToken.None);
+		List<(AgentId Id, int Count)> second = await SessionConflictService.DrainConflictsAsync();
+		Assert.AreEqual(0, second.Count);
+	}
 
-			IAgent? updated = await AgentCollection.GetAsync(agent.Id);
-			Assert.IsNotNull(updated);
-			Assert.IsTrue(updated.Enabled, "Agent should remain enabled when feature is disabled");
-		}
+	[TestMethod]
+	public async Task DrainEmptyReturnsEmptyListAsync()
+	{
+		List<(AgentId Id, int Count)> conflicts = await SessionConflictService.DrainConflictsAsync();
+		Assert.AreEqual(0, conflicts.Count);
 	}
 }
diff --git a/Engine/Source/Programs/Horde/Plugins/Compute/HordeServer.Compute/Agents/SessionConflictService.cs b/Engine/Source/Programs/Horde/Plugins/Compute/HordeServer.Compute/Agents/SessionConflictService.cs
--- a/Engine/Source/Programs/Horde/Plugins/Compute/HordeServer.Compute/Agents/SessionConflictService.cs
+++ b/Engine/Source/Programs/Horde/Plugins/Compute/HordeServer.Compute/Agents/SessionConflictService.cs
@@ -1,131 +1,61 @@
 // Copyright Epic Games, Inc. All Rights Reserved.
 
-using System.Net;
 using EpicGames.Horde.Agents;
-using EpicGames.Horde.Agents.Sessions;
 using HordeServer.Server;
-using Microsoft.Extensions.Logging;
-using Microsoft.Extensions.Options;
 using StackExchange.Redis;
 
 namespace HordeServer.Agents
 {
 	/// <summary>
-	/// Detects session conflict flip-flops (e.g. cloned VMs sharing an agent ID)
-	/// and auto-quarantines the agent after a configurable threshold.
+	/// Tracks session conflict counts per agent in a Redis sorted set.
+	/// Consumed by NotificationService to produce a 12-hour Slack report.
 	/// </summary>
 	public class SessionConflictService
 	{
-		/// <summary>
-		/// Prefix used in agent comments when quarantining due to session conflicts.
-		/// Also used by AgentReportService to identify these agents.
-		/// </summary>
-		public const string CommentPrefix = "[session-conflict]";
+		static readonly RedisKey s_redisKey = new("session-conflicts");
+		static readonly TimeSpan s_keyTtl = TimeSpan.FromHours(48);
 
 		readonly IRedisService _redisService;
-		readonly IAgentCollection _agentCollection;
-		readonly IOptionsMonitor<ComputeServerConfig> _serverConfig;
-		readonly ILogger _logger;
 
 		/// <summary>
 		/// Constructor
 		/// </summary>
-		public SessionConflictService(IRedisService redisService, IAgentCollection agentCollection, IOptionsMonitor<ComputeServerConfig> serverConfig, ILogger<SessionConflictService> logger)
-		{
-			_redisService = redisService;
-			_agentCollection = agentCollection;
-			_serverConfig = serverConfig;
-			_logger = logger;
-		}
+		public SessionConflictService(IRedisService redisService)
+			=> _redisService = redisService;
 
-		internal static RedisKey GetRedisKey(AgentId agentId) => new RedisKey($"session-conflict:{agentId}");
-
-		// Member format: "{staleSessionId}:{currentSessionId}:{sanitizedIp}:{uniqueId}"
-		// SessionId is a hex BinaryId (no colons), and IPv6 colons are replaced with underscores.
-		static string CreateMember(SessionId staleSessionId, SessionId? currentSessionId, IPAddress? remoteIp)
-		{
-			string sanitizedIp = remoteIp?.ToString().Replace(':', '_') ?? "unknown";
-			return $"{staleSessionId}:{currentSessionId}:{sanitizedIp}:{Guid.NewGuid():N}";
-		}
-
-		static string ParseIpFromMember(string member)
+		/// <summary>
+		/// Record a session conflict for the given agent.
+		/// </summary>
+		/// <returns>The total conflict count for this agent since the last drain.</returns>
+		public async Task<long> RecordConflictAsync(AgentId agentId, CancellationToken cancellationToken = default)
 		{
-			string[] parts = member.Split(':');
-			return parts.Length >= 3 ? parts[2].Replace('_', ':') : "unknown";
+			_ = cancellationToken; // Redis client doesn't accept CT directly; kept for API consistency
+			IDatabase db = _redisService.GetDatabase();
+			long count = (long)await db.SortedSetIncrementAsync(s_redisKey, agentId.ToString(), 1.0);
+			await db.KeyExpireAsync(s_redisKey, s_keyTtl, CommandFlags.FireAndForget);
+			return count;
 		}
 
 		/// <summary>
-		/// Records a session conflict event and quarantines the agent if the threshold is exceeded.
+		/// Drains all recorded conflicts and clears the set.
+		/// Returns each agent ID paired with its total mismatch count since the last drain.
 		/// </summary>
-		/// <param name="agentId">The agent experiencing the conflict</param>
-		/// <param name="staleSessionId">The session ID that the caller tried to use</param>
-		/// <param name="currentSessionId">The actual current session ID on the agent</param>
-		/// <param name="remoteIp">IP address of the caller</param>
-		/// <param name="cancellationToken">Cancellation token</param>
-		public async Task RecordConflictAsync(AgentId agentId, SessionId staleSessionId, SessionId? currentSessionId, IPAddress? remoteIp, CancellationToken cancellationToken)
+		public async Task<List<(AgentId Id, int Count)>> DrainConflictsAsync(CancellationToken cancellationToken = default)
 		{
-			ComputeServerConfig config = _serverConfig.CurrentValue;
-			if (!config.SessionConflictEnabled)
-			{
-				return;
-			}
-
+			_ = cancellationToken;
 			IDatabase db = _redisService.GetDatabase();
-			RedisKey key = GetRedisKey(agentId);
-			TimeSpan window = TimeSpan.FromSeconds(config.SessionConflictWindowSeconds);
-
-			long nowTicks = DateTimeOffset.UtcNow.Ticks;
-			string member = CreateMember(staleSessionId, currentSessionId, remoteIp);
-
-			// Add event and prune old entries
-			await db.SortedSetAddAsync(key, member, nowTicks);
-			await db.SortedSetRemoveRangeByScoreAsync(key, double.NegativeInfinity, nowTicks - window.Ticks);
-			TimeSpan ttl = window + window;
-			await db.KeyExpireAsync(key, ttl);
-
-			long count = await db.SortedSetLengthAsync(key);
-			if (count < config.SessionConflictThreshold)
-			{
-				return;
-			}
 
-			// Collect distinct IPs from recent events for the quarantine comment
-			SortedSetEntry[] entries = await db.SortedSetRangeByScoreWithScoresAsync(key);
-			HashSet<string> ips = new HashSet<string>(StringComparer.Ordinal);
-			foreach (SortedSetEntry entry in entries)
-			{
-				ips.Add(ParseIpFromMember(entry.Element.ToString()));
-			}
-
-			string ipList = String.Join(", ", ips.Order());
-			string comment = $"{CommentPrefix} Auto-quarantined: session flip-flop detected ({count} conflicts). IPs: {ipList}. Time: {DateTimeOffset.UtcNow:O}. Re-enable after investigating duplicate agent IDs.";
+			ITransaction transaction = db.CreateTransaction();
+			Task<SortedSetEntry[]> entriesTask = transaction.SortedSetRangeByScoreWithScoresAsync(s_redisKey);
+			_ = transaction.KeyDeleteAsync(s_redisKey, CommandFlags.FireAndForget);
 
-			// Quarantine the agent with retry for optimistic concurrency
-			for (int attempt = 0; attempt < 5; attempt++)
+			if (!await transaction.ExecuteAsync())
 			{
-				IAgent? agent = await _agentCollection.GetAsync(agentId, cancellationToken);
-				if (agent == null)
-				{
-					break;
-				}
-
-				if (!agent.Enabled)
-				{
-					break; // Already disabled, nothing to do
-				}
-
-				IAgent? updated = await agent.TryUpdateAsync(new UpdateAgentOptions(Enabled: false, Comment: comment), cancellationToken);
-				if (updated != null)
-				{
-					_logger.LogWarning("Agent {AgentId} quarantined due to session conflict flip-flop ({Count} conflicts in {Window}s). IPs: {IpList}", agentId, count, config.SessionConflictWindowSeconds, ipList);
-
-					// Clean up Redis key after successful quarantine
-					await db.KeyDeleteAsync(key);
-					return;
-				}
+				return [];
 			}
 
-			_logger.LogWarning("Failed to quarantine agent {AgentId} after 5 attempts", agentId);
+			SortedSetEntry[] entries = await entriesTask;
+			return entries.Select(e => (new AgentId(e.Element.ToString()), (int)e.Score)).ToList();
 		}
 	}
 }
diff --git a/Engine/Source/Programs/Horde/Plugins/Compute/HordeServer.Compute/ComputeServerConfig.cs b/Engine/Source/Programs/Horde/Plugins/Compute/HordeServer.Compute/ComputeServerConfig.cs
--- a/Engine/Source/Programs/Horde/Plugins/Compute/HordeServer.Compute/ComputeServerConfig.cs
+++ b/Engine/Source/Programs/Horde/Plugins/Compute/HordeServer.Compute/ComputeServerConfig.cs
@@ -67,23 +67,6 @@ namespace HordeServer
 		/// </summary>
 		public int AgentPoolScaleInCooldownSeconds { get; set; } = 1200; // 20 mins
 
-		/// <summary>
-		/// Whether to enable session conflict detection and auto-quarantine.
-		/// When multiple machines accidentally connect using the same agent ID (e.g. cloned VMs),
-		/// repeated session mismatches will cause the agent to be automatically disabled.
-		/// </summary>
-		public bool SessionConflictEnabled { get; set; } = false;
-
-		/// <summary>
-		/// Number of session mismatches within the sliding window required to trigger auto-quarantine
-		/// </summary>
-		public int SessionConflictThreshold { get; set; } = 3;
-
-		/// <summary>
-		/// Sliding window in seconds for counting session conflict events
-		/// </summary>
-		public int SessionConflictWindowSeconds { get; set; } = 300;
-
 		/// <summary>
 		/// Port to listen on for tunneling compute sockets to agents
 		/// </summary>
diff --git a/Engine/Source/Programs/Horde/Plugins/Compute/HordeServer.Compute/Server/RpcService.cs b/Engine/Source/Programs/Horde/Plugins/Compute/HordeServer.Compute/Server/RpcService.cs
--- a/Engine/Source/Programs/Horde/Plugins/Compute/HordeServer.Compute/Server/RpcService.cs
+++ b/Engine/Source/Programs/Horde/Plugins/Compute/HordeServer.Compute/Server/RpcService.cs
@@ -33,7 +33,13 @@ namespace HordeServer.Server
 	public class RpcService : HordeRpc.HordeRpcBase
 	{
 		/// <summary>
-		/// Timeout before closing a long-polling request (client will retry again) 
+		/// Number of session mismatches per agent that generate detailed (Error-level) structured exceptions.
+		/// Beyond this threshold, mismatches are logged at Debug level to suppress spam.
+		/// </summary>
+		internal const int MaxDetailedConflictErrors = 20;
+
+		/// <summary>
+		/// Timeout before closing a long-polling request (client will retry again)
 		/// </summary>
 		internal TimeSpan _longPollTimeout = TimeSpan.FromMinutes(9);
 
@@ -44,7 +50,6 @@ namespace HordeServer.Server
 		readonly IAclService _aclService;
 		readonly IOptionsSnapshot<ComputeConfig> _computeConfig;
 		readonly SessionConflictService _sessionConflictService;
-		readonly IOptionsMonitor<ComputeServerConfig> _staticComputeConfig;
 		readonly ILogger _logger;
 		readonly ITelemetryWriter _telemetryWriter;
 
@@ -59,7 +64,6 @@ namespace HordeServer.Server
 			IAclService aclService,
 			IOptionsSnapshot<ComputeConfig> computeConfig,
 			SessionConflictService sessionConflictService,
-			IOptionsMonitor<ComputeServerConfig> staticComputeConfig,
 			ILogger<RpcService> logger,
 			ITelemetryWriter telemetryWriter)
 		{
@@ -70,7 +74,6 @@ namespace HordeServer.Server
 			_aclService = aclService;
 			_computeConfig = computeConfig;
 			_sessionConflictService = sessionConflictService;
-			_staticComputeConfig = staticComputeConfig;
 			_logger = logger;
 			_telemetryWriter = telemetryWriter;
 		}
@@ -318,23 +321,16 @@ namespace HordeServer.Server
 						// Check we're authorized to update it
 						if (agent.SessionId != sessionId)
 						{
-							IPAddress? remoteIp = context.GetHttpContext().Connection?.RemoteIpAddress;
-							bool conflictEnabled = _staticComputeConfig.CurrentValue.SessionConflictEnabled;
-							bool alreadyQuarantined = conflictEnabled && !agent.Enabled && agent.Comment?.StartsWith(SessionConflictService.CommentPrefix, StringComparison.Ordinal) == true;
-
-							if (alreadyQuarantined)
+							long conflictCount = await _sessionConflictService.RecordConflictAsync(agent.Id);
+							if (conflictCount <= MaxDetailedConflictErrors)
 							{
-								// Suppress log spam: log at Debug and throw plain RpcException (bypasses the interceptor's ERROR log).
-								_logger.LogDebug("Session mismatch for already-quarantined agent {AgentId}: caller has {StaleSessionId} but current is {CurrentSessionId}, remote IP {RemoteIp}", agent.Id, sessionId, agent.SessionId, remoteIp);
-								throw new RpcException(new Status(StatusCode.PermissionDenied, $"Agent {request.AgentId} is quarantined for session conflicts."));
+								throw new StructuredRpcException(StatusCode.PermissionDenied, "Agent {AgentId} has completed session {SessionId}; now executing session {NewSessionId}. Cannot update state.", request.AgentId, sessionId, agent.SessionId?.ToString() ?? "(None)");
 							}
-
-							_logger.LogInformation("Session mismatch for agent {AgentId}: caller has {StaleSessionId} but current is {CurrentSessionId}, remote IP {RemoteIp}", agent.Id, sessionId, agent.SessionId, remoteIp);
-							if (conflictEnabled)
+							else
 							{
-								await _sessionConflictService.RecordConflictAsync(agent.Id, sessionId, agent.SessionId, remoteIp, context.CancellationToken);
+								_logger.LogDebug("Agent {AgentId} session mismatch (count: {ConflictCount}): expected {SessionId}, actual {NewSessionId}", request.AgentId, conflictCount, sessionId, agent.SessionId?.ToString() ?? "(None)");
+								throw new RpcException(new Status(StatusCode.PermissionDenied, $"Agent {request.AgentId} session mismatch"));
 							}
-							throw new StructuredRpcException(StatusCode.PermissionDenied, "Agent {AgentId} has completed session {SessionId}; now executing session {NewSessionId}. Cannot update state.", request.AgentId, sessionId, agent.SessionId?.ToString() ?? "(None)");
 						}
 						if (!_agentService.AuthorizeSession(agent, context.GetHttpContext().User, out string authReason))
 						{
diff --git a/Engine/Source/Programs/Horde/Plugins/Experimental/HordeServer.Experimental.Tests/Notifications/ExperimentalNotificationTests.cs b/Engine/Source/Programs/Horde/Plugins/Experimental/HordeServer.Experimental.Tests/Notifications/ExperimentalNotificationTests.cs
--- a/Engine/Source/Programs/Horde/Plugins/Experimental/HordeServer.Experimental.Tests/Notifications/ExperimentalNotificationTests.cs
+++ b/Engine/Source/Programs/Horde/Plugins/Experimental/HordeServer.Experimental.Tests/Notifications/ExperimentalNotificationTests.cs
@@ -126,6 +126,7 @@ namespace HordeServer.Experimental.Tests.Notifications
 		public Task NotifyConfigUpdateFailureAsync(string errorMessage, string fileName, int? change = null, IUser? author = null, string? description = null, CancellationToken cancellationToken = default) { throw new NotImplementedException(); }
 		public Task NotifyDeviceServiceAsync(string message, IDevice? device = null, IDevicePool? pool = null, StreamConfig? stream = null, IJob? job = null, IJobStep? step = null, INode? node = null, IUser? user = null, CancellationToken cancellationToken = default) { throw new NotImplementedException(); }
 		public Task SendAgentReportAsync(AgentReport report, CancellationToken cancellationToken) => throw new NotImplementedException();
+		public Task SendSessionConflictReportAsync(IReadOnlyList<(AgentId Id, int Count)> conflicts, CancellationToken cancellationToken) => Task.CompletedTask;
 		public Task SendIssueReportAsync(IssueReportGroup report, CancellationToken cancellationToken) => throw new NotImplementedException();
 		public Task SendDeviceIssueReportAsync(DeviceIssueReport report, CancellationToken cancellationToken) => throw new NotImplementedException();
 		public Task NotifyTestHealthReportAsync(ITestHealthReport report, string recipient, string[]? carbonCopies, CancellationToken cancellationToken) => throw new NotImplementedException();
diff --git a/Engine/Source/Programs/Horde/Plugins/Experimental/HordeServer.Experimental/Notifications/Sinks/ExperimentalSlackNotificationSink.cs b/Engine/Source/Programs/Horde/Plugins/Experimental/HordeServer.Experimental/Notifications/Sinks/ExperimentalSlackNotificationSink.cs
--- a/Engine/Source/Programs/Horde/Plugins/Experimental/HordeServer.Experimental/Notifications/Sinks/ExperimentalSlackNotificationSink.cs
+++ b/Engine/Source/Programs/Horde/Plugins/Experimental/HordeServer.Experimental/Notifications/Sinks/ExperimentalSlackNotificationSink.cs
@@ -1,5 +1,6 @@
 // Copyright Epic Games, Inc. All Rights Reserved.
 
+using EpicGames.Horde.Agents;
 using EpicGames.Horde.Jobs;
 using EpicGames.Horde.Jobs.Graphs;
 using EpicGames.Horde.Users;
@@ -110,6 +111,11 @@ namespace HordeServer.Notifications.Sinks
 		public Task SendAgentReportAsync(AgentReport report, CancellationToken cancellationToken)
 			=> Task.CompletedTask;
 
+		/// <inheritdoc/>
+		/// <remarks>No-op: this sink only handles job completion/step notifications via the processor. Session conflict reports are sent by the primary SlackNotificationSink.</remarks>
+		public Task SendSessionConflictReportAsync(IReadOnlyList<(AgentId Id, int Count)> conflicts, CancellationToken cancellationToken)
+			=> Task.CompletedTask;
+
 		/// <inheritdoc/>
 		public Task SendDeviceIssueReportAsync(DeviceIssueReport report, CancellationToken cancellationToken)
 			=> Task.CompletedTask;
