[test automation hub] Avoid test phase name to be collapsable

#rnx

#Synced-CL 51758956

diff --git a/Engine/Source/Programs/Horde/HordeDashboard/plugins/test-automation/testAutomationPhaseGridView.tsx b/Engine/Source/Programs/Horde/HordeDashboard/plugins/test-automation/testAutomationPhaseGridView.tsx
--- a/Engine/Source/Programs/Horde/HordeDashboard/plugins/test-automation/testAutomationPhaseGridView.tsx
+++ b/Engine/Source/Programs/Horde/HordeDashboard/plugins/test-automation/testAutomationPhaseGridView.tsx
@@ -494,7 +494,7 @@ const getPhaseColumn = (context: MetaSelectorContext, handler: TestDataHandler,
             name: 'Phase',
             minWidth: 200,
             isResizable: true,
-            isCollapsible: true,
+            isCollapsible: false,
             isSorted: false,
             isSortedDescending: true,
             onRender: onPhaseNameRender,
