Move Robomerge config from server config to global config for dynamic reload via globals.json

diff --git a/Engine/Source/Programs/Horde/Plugins/Build/HordeServer.Build/BuildConfig.cs b/Engine/Source/Programs/Horde/Plugins/Build/HordeServer.Build/BuildConfig.cs
--- a/Engine/Source/Programs/Horde/Plugins/Build/HordeServer.Build/BuildConfig.cs
+++ b/Engine/Source/Programs/Horde/Plugins/Build/HordeServer.Build/BuildConfig.cs
@@ -64,6 +64,12 @@ namespace HordeServer
 		/// </summary>
 		public TimeSpan? AgentShutdownIfDisabledGracePeriod { get; set; } = null;
 
+		/// <summary>
+		/// RoboMerge branchmap file configuration for the merge graph service.
+		/// When null or empty, the service is inert (no Perforce connections, no API data).
+		/// </summary>
+		public RoboMergeConfig? Robomerge { get; set; }
+
 		/// <summary>
 		/// Configuration for different artifact types
 		/// </summary>
@@ -623,4 +629,30 @@ namespace HordeServer
 		/// </summary>
 		public List<AutoSdkWorkspace> AutoSdk { get; set; } = new List<AutoSdkWorkspace>();
 	}
+
+	/// <summary>
+	/// Configuration for the RoboMerge branchmap service
+	/// </summary>
+	public class RoboMergeConfig
+	{
+		/// <summary>
+		/// Perforce depot paths to branchmap JSON files.
+		/// Plain depot paths (e.g. "//MyProject/Main/Build/Streams/myproject.branchmap.json").
+		/// All files are read from the Perforce cluster specified by <see cref="PerforceCluster"/>.
+		/// </summary>
+		public List<string> BranchmapFiles { get; set; } = new List<string>();
+
+		/// <summary>
+		/// Perforce cluster name to connect to for reading branchmap files.
+		/// Maps to a connection configured in BuildServerConfig.Perforce[].
+		/// Default: "default" (uses the default Perforce connection).
+		/// </summary>
+		public string PerforceCluster { get; set; } = "default";
+
+		/// <summary>
+		/// How often to poll Perforce for branchmap file changes. Default: 5 minutes.
+		/// Changing this requires a server restart (ticker interval is set at construction time).
+		/// </summary>
+		public TimeSpan TickInterval { get; set; } = TimeSpan.FromMinutes(5);
+	}
 }
diff --git a/Engine/Source/Programs/Horde/Plugins/Build/HordeServer.Build/BuildServerConfig.cs b/Engine/Source/Programs/Horde/Plugins/Build/HordeServer.Build/BuildServerConfig.cs
--- a/Engine/Source/Programs/Horde/Plugins/Build/HordeServer.Build/BuildServerConfig.cs
+++ b/Engine/Source/Programs/Horde/Plugins/Build/HordeServer.Build/BuildServerConfig.cs
@@ -45,13 +45,6 @@ namespace HordeServer
 		/// </summary>
 		public Uri? RobomergeUrl { get; set; }
 
-		/// <summary>
-		/// RoboMerge branchmap file configuration for the merge graph service.
-		/// When null or empty, the service is inert (no Perforce connections, no API data).
-		/// This is separate from <see cref="RobomergeUrl"/> which is the URL to the RoboMerge web UI for linking purposes.
-		/// </summary>
-		public RoboMergeConfig? Robomerge { get; set; }
-
 		/// <summary>
 		/// Url of Commits Viewer
 		/// </summary>
@@ -273,32 +266,6 @@ namespace HordeServer
 		}
 	}
 
-	/// <summary>
-	/// Configuration for the RoboMerge branchmap service
-	/// </summary>
-	public class RoboMergeConfig
-	{
-		/// <summary>
-		/// Perforce depot paths to branchmap JSON files.
-		/// Plain depot paths (e.g. "//MyProject/Main/Build/Streams/myproject.branchmap.json").
-		/// All files are read from the Perforce cluster specified by <see cref="PerforceCluster"/>.
-		/// </summary>
-		public List<string> BranchmapFiles { get; set; } = new List<string>();
-
-		/// <summary>
-		/// Perforce cluster name to connect to for reading branchmap files.
-		/// Maps to a connection configured in BuildServerConfig.Perforce[].
-		/// Default: "default" (uses the default Perforce connection).
-		/// </summary>
-		public string PerforceCluster { get; set; } = "default";
-
-		/// <summary>
-		/// How often to poll Perforce for branchmap file changes. Default: 5 minutes.
-		/// Changing this requires a server restart (ticker interval is set at construction time).
-		/// </summary>
-		public TimeSpan TickInterval { get; set; } = TimeSpan.FromMinutes(5);
-	}
-
 	/// <summary>
 	/// Options for the commit service
 	/// </summary>
diff --git a/Engine/Source/Programs/Horde/Plugins/Build/HordeServer.Build/RoboMerge/RoboMergeController.cs b/Engine/Source/Programs/Horde/Plugins/Build/HordeServer.Build/RoboMerge/RoboMergeController.cs
--- a/Engine/Source/Programs/Horde/Plugins/Build/HordeServer.Build/RoboMerge/RoboMergeController.cs
+++ b/Engine/Source/Programs/Horde/Plugins/Build/HordeServer.Build/RoboMerge/RoboMergeController.cs
@@ -18,16 +18,14 @@ namespace HordeServer.RoboMerge
 	public class RoboMergeController : HordeControllerBase
 	{
 		readonly IMergeGraphService _mergeGraphService;
-		readonly IOptionsMonitor<BuildServerConfig> _serverConfig;
 		readonly IOptionsSnapshot<BuildConfig> _buildConfig;
 
 		/// <summary>
 		/// Constructor
 		/// </summary>
-		public RoboMergeController(IMergeGraphService mergeGraphService, IOptionsMonitor<BuildServerConfig> serverConfig, IOptionsSnapshot<BuildConfig> buildConfig)
+		public RoboMergeController(IMergeGraphService mergeGraphService, IOptionsSnapshot<BuildConfig> buildConfig)
 		{
 			_mergeGraphService = mergeGraphService;
-			_serverConfig = serverConfig;
 			_buildConfig = buildConfig;
 		}
 
@@ -76,7 +74,7 @@ namespace HordeServer.RoboMerge
 			IReadOnlyList<MergeGraph> graphs = _mergeGraphService.GetMergeGraphs();
 			return new GetRoboMergeStatusResponse
 			{
-				IsConfigured = _serverConfig.CurrentValue.Robomerge?.BranchmapFiles?.Count > 0,
+				IsConfigured = _buildConfig.Value.Robomerge?.BranchmapFiles?.Count > 0,
 				GraphsLoaded = graphs.Count > 0,
 				GraphCount = graphs.Count,
 				BotNames = graphs.Select(g => g.BotName).ToList(),
diff --git a/Engine/Source/Programs/Horde/Plugins/Build/HordeServer.Build/RoboMerge/RoboMergeService.cs b/Engine/Source/Programs/Horde/Plugins/Build/HordeServer.Build/RoboMerge/RoboMergeService.cs
--- a/Engine/Source/Programs/Horde/Plugins/Build/HordeServer.Build/RoboMerge/RoboMergeService.cs
+++ b/Engine/Source/Programs/Horde/Plugins/Build/HordeServer.Build/RoboMerge/RoboMergeService.cs
@@ -67,7 +67,7 @@ namespace HordeServer.RoboMerge
 	internal sealed class RoboMergeService : IHostedService, IAsyncDisposable, IMergeGraphService
 	{
 		readonly IPerforceService _perforceService;
-		readonly IOptionsMonitor<BuildServerConfig> _serverConfig;
+		readonly IOptionsMonitor<BuildConfig> _buildConfig;
 		readonly ILogger<RoboMergeService> _logger;
 
 		readonly ITicker _ticker;
@@ -99,16 +99,16 @@ namespace HordeServer.RoboMerge
 		/// </summary>
 		public RoboMergeService(
 			IPerforceService perforceService,
-			IOptionsMonitor<BuildServerConfig> serverConfig,
+			IOptionsMonitor<BuildConfig> buildConfig,
 			IClock clock,
 			Meter meter,
 			ILogger<RoboMergeService> logger)
 		{
 			_perforceService = perforceService;
-			_serverConfig = serverConfig;
+			_buildConfig = buildConfig;
 			_logger = logger;
 
-			_tickInterval = serverConfig.CurrentValue.Robomerge?.TickInterval ?? TimeSpan.FromMinutes(5);
+			_tickInterval = buildConfig.CurrentValue.Robomerge?.TickInterval ?? TimeSpan.FromMinutes(5);
 
 			_ticker = clock.AddTicker<RoboMergeService>(s_initialDelay, TickInternalAsync, logger);
 
@@ -145,7 +145,7 @@ namespace HordeServer.RoboMerge
 			Stopwatch sw = Stopwatch.StartNew();
 			try
 			{
-				BuildServerConfig config = _serverConfig.CurrentValue;
+				BuildConfig config = _buildConfig.CurrentValue;
 				List<string>? files = config.Robomerge?.BranchmapFiles;
 				if (files == null || files.Count == 0)
 				{
