Horde: Remove unused using directive in JobsController

diff --git a/Engine/Source/Programs/Horde/Plugins/Build/HordeServer.Build/Jobs/JobsController.cs b/Engine/Source/Programs/Horde/Plugins/Build/HordeServer.Build/Jobs/JobsController.cs
--- a/Engine/Source/Programs/Horde/Plugins/Build/HordeServer.Build/Jobs/JobsController.cs
+++ b/Engine/Source/Programs/Horde/Plugins/Build/HordeServer.Build/Jobs/JobsController.cs
@@ -12,7 +12,6 @@ using EpicGames.Horde.Logs;
 using EpicGames.Horde.Notifications;
 using EpicGames.Horde.Streams;
 using EpicGames.Horde.Users;
-using EpicGames.Horde.Issues;
 using HordeServer.Auditing;
 using HordeServer.Acls;
 using HordeServer.Agents;
