Horde: Updating schema documentation

#rnx

diff --git a/Engine/Source/Programs/Horde/Docs/Config/Schema/AclActions.md b/Engine/Source/Programs/Horde/Docs/Config/Schema/AclActions.md
--- a/Engine/Source/Programs/Horde/Docs/Config/Schema/AclActions.md
+++ b/Engine/Source/Programs/Horde/Docs/Config/Schema/AclActions.md
@@ -163,6 +163,30 @@
 | `ViewAgent` |  |
 | `ListAgents` |  |
 
+## Artifacts
+
+| Name | Description |
+| ---- | ----------- |
+| `ReadArtifact` | Permission to read from an artifact |
+| `WriteArtifact` | Permission to write to an artifact |
+| `DeleteArtifact` | Permission to delete to an artifact |
+| `UploadArtifact` | Ability to create an artifact. Typically just for debugging; agents have this access for a particular session. |
+| `DownloadArtifact` | Ability to download an artifact |
+
+## Bisect
+
+| Name | Description |
+| ---- | ----------- |
+| `CreateBisectTask` | Ability to start new bisect tasks |
+| `UpdateBisectTask` | Ability to update a bisect task |
+| `ViewBisectTask` | Ability to view a bisect task |
+
+## BuildHealth
+
+| Name | Description |
+| ---- | ----------- |
+| `AddUpdateDeleteBuildHealthFilter` | Ability to add, update, and delete build health filters. |
+
 ## Compute
 
 | Name | Description |
@@ -172,6 +196,24 @@
 | `UbaCacheRead` |  |
 | `UbaCacheWrite` |  |
 
+## Devices
+
+| Name | Description |
+| ---- | ----------- |
+| `DeviceRead` | Ability to read devices |
+| `DeviceWrite` | Ability to write devices |
+
+## Jobs
+
+| Name | Description |
+| ---- | ----------- |
+| `CreateJob` | Ability to start new jobs |
+| `UpdateJob` | Rename a job, modify its priority, etc... |
+| `DeleteJob` | Delete a job properties |
+| `ExecuteJob` | Allows updating a job metadata (name, changelist number, step properties, new groups, job states, etc...). Typically granted to agents. Not user facing. |
+| `RetryJobStep` | Ability to retry a failed job step |
+| `ViewJob` | Ability to view a job |
+
 ## Leases
 
 | Name | Description |
@@ -198,6 +240,12 @@
 | `UpdateNotice` |  |
 | `DeleteNotice` |  |
 
+## Notifications
+
+| Name | Description |
+| ---- | ----------- |
+| `CreateSubscription` | Ability to subscribe to notifications |
+
 ## Pools
 
 | Name | Description |
@@ -208,6 +256,22 @@
 | `ViewPool` |  |
 | `ListPools` |  |
 
+## Projects
+
+| Name | Description |
+| ---- | ----------- |
+| `CreateProject` | Allows the creation of new projects |
+| `DeleteProject` | Allows deletion of projects. |
+| `UpdateProject` | Modify attributes of a project (name, categories, etc...) |
+| `ViewProject` | View information about a project |
+
+## Replicators
+
+| Name | Description |
+| ---- | ----------- |
+| `UpdateReplicator` | Allows deletion of projects. |
+| `ViewReplicator` | Allows the creation of new projects |
+
 ## ServiceAccounts
 
 | Name | Description |
@@ -232,6 +296,24 @@
 | `DownloadSoftware` |  |
 | `DeleteSoftware` |  |
 
+## Streams
+
+| Name | Description |
+| ---- | ----------- |
+| `CreateStream` | Allows the creation of new streams within a project |
+| `UpdateStream` | Allows updating a stream (agent types, templates, schedules) |
+| `DeleteStream` | Allows deleting a stream |
+| `ViewStream` | Ability to view a stream |
+| `ViewChanges` | View changes submitted to a stream. NOTE: this returns responses from the server's Perforce account, which may be a priviledged user. |
+| `ViewTemplate` | View template associated with a stream |
+
+## TestData
+
+| Name | Description |
+| ---- | ----------- |
+| `TestHealthRead` | Ability to read test health |
+| `TestHealthWrite` | Ability to write test health |
+
 ## Accounts
 
 | Name | Description |
@@ -256,25 +338,25 @@
 
 | Name | Description |
 | ---- | ----------- |
-| `ReadArtifact` | Permission to read from an artifact |
-| `WriteArtifact` | Permission to write to an artifact |
-| `DeleteArtifact` | Permission to delete to an artifact |
-| `UploadArtifact` | Ability to create an artifact. Typically just for debugging; agents have this access for a particular session. |
-| `DownloadArtifact` | Ability to download an artifact |
+| `ReadArtifact` |  |
+| `WriteArtifact` |  |
+| `DeleteArtifact` |  |
+| `UploadArtifact` |  |
+| `DownloadArtifact` |  |
 
 ## Bisect
 
 | Name | Description |
 | ---- | ----------- |
-| `CreateBisectTask` | Ability to start new bisect tasks |
-| `UpdateBisectTask` | Ability to update a bisect task |
-| `ViewBisectTask` | Ability to view a bisect task |
+| `CreateBisectTask` |  |
+| `UpdateBisectTask` |  |
+| `ViewBisectTask` |  |
 
 ## BuildHealth
 
 | Name | Description |
 | ---- | ----------- |
-| `AddUpdateDeleteBuildHealthFilter` | Ability to add, update, and delete build health filters. |
+| `AddUpdateDeleteBuildHealthFilter` |  |
 
 ## Compute
 
@@ -289,19 +371,19 @@
 
 | Name | Description |
 | ---- | ----------- |
-| `DeviceRead` | Ability to read devices |
-| `DeviceWrite` | Ability to write devices |
+| `DeviceRead` |  |
+| `DeviceWrite` |  |
 
 ## Jobs
 
 | Name | Description |
 | ---- | ----------- |
-| `CreateJob` | Ability to start new jobs |
-| `UpdateJob` | Rename a job, modify its priority, etc... |
-| `DeleteJob` | Delete a job properties |
-| `ExecuteJob` | Allows updating a job metadata (name, changelist number, step properties, new groups, job states, etc...). Typically granted to agents. Not user facing. |
-| `RetryJobStep` | Ability to retry a failed job step |
-| `ViewJob` | Ability to view a job |
+| `CreateJob` |  |
+| `UpdateJob` |  |
+| `DeleteJob` |  |
+| `ExecuteJob` |  |
+| `RetryJobStep` |  |
+| `ViewJob` |  |
 
 ## Leases
 
@@ -333,7 +415,7 @@
 
 | Name | Description |
 | ---- | ----------- |
-| `CreateSubscription` | Ability to subscribe to notifications |
+| `CreateSubscription` |  |
 
 ## Pools
 
@@ -349,17 +431,17 @@
 
 | Name | Description |
 | ---- | ----------- |
-| `CreateProject` | Allows the creation of new projects |
-| `DeleteProject` | Allows deletion of projects. |
-| `UpdateProject` | Modify attributes of a project (name, categories, etc...) |
-| `ViewProject` | View information about a project |
+| `CreateProject` |  |
+| `DeleteProject` |  |
+| `UpdateProject` |  |
+| `ViewProject` |  |
 
 ## Replicators
 
 | Name | Description |
 | ---- | ----------- |
-| `UpdateReplicator` | Allows deletion of projects. |
-| `ViewReplicator` | Allows the creation of new projects |
+| `UpdateReplicator` |  |
+| `ViewReplicator` |  |
 
 ## ServiceAccounts
 
@@ -389,19 +471,19 @@
 
 | Name | Description |
 | ---- | ----------- |
-| `CreateStream` | Allows the creation of new streams within a project |
-| `UpdateStream` | Allows updating a stream (agent types, templates, schedules) |
-| `DeleteStream` | Allows deleting a stream |
-| `ViewStream` | Ability to view a stream |
-| `ViewChanges` | View changes submitted to a stream. NOTE: this returns responses from the server's Perforce account, which may be a priviledged user. |
-| `ViewTemplate` | View template associated with a stream |
+| `CreateStream` |  |
+| `UpdateStream` |  |
+| `DeleteStream` |  |
+| `ViewStream` |  |
+| `ViewChanges` |  |
+| `ViewTemplate` |  |
 
 ## TestData
 
 | Name | Description |
 | ---- | ----------- |
-| `TestHealthRead` | Ability to read test health |
-| `TestHealthWrite` | Ability to write test health |
+| `TestHealthRead` |  |
+| `TestHealthWrite` |  |
 
 ## Accounts
 
@@ -532,6 +614,12 @@
 | `UpdateReplicator` |  |
 | `ViewReplicator` |  |
 
+## Secrets
+
+| Name | Description |
+| ---- | ----------- |
+| `ViewSecret` | View a credential |
+
 ## ServiceAccounts
 
 | Name | Description |
@@ -707,7 +795,7 @@
 
 | Name | Description |
 | ---- | ----------- |
-| `ViewSecret` | View a credential |
+| `ViewSecret` |  |
 
 ## ServiceAccounts
 
diff --git a/Engine/Source/Programs/Horde/Docs/Config/Schema/Globals.md b/Engine/Source/Programs/Horde/Docs/Config/Schema/Globals.md
--- a/Engine/Source/Programs/Horde/Docs/Config/Schema/Globals.md
+++ b/Engine/Source/Programs/Horde/Docs/Config/Schema/Globals.md
@@ -61,10 +61,10 @@ Name | Description
 ---- | -----------
 `analytics` | [AnalyticsConfig](#analyticsconfig)<br>Configuration for the analytics plugin
 `compute` | [ComputeConfig](#computeconfig)<br>Configuration for the compute plugin
-`exampleplugin` | [ExamplePluginConfig](#examplepluginconfig)<br>Configuration for the exampleplugin plugin
 `build` | [BuildConfig](#buildconfig)<br>Configuration for the build plugin
-`experimental` | [ExperimentalConfig](#experimentalconfig)<br>Configuration for the experimental plugin
+`exampleplugin` | [ExamplePluginConfig](#examplepluginconfig)<br>Configuration for the exampleplugin plugin
 `secrets` | [SecretsConfig](#secretsconfig)<br>Configuration for the secrets plugin
+`experimental` | [ExperimentalConfig](#experimentalconfig)<br>Configuration for the experimental plugin
 `symbols` | [SymbolsConfig](#symbolsconfig)<br>Configuration for the symbols plugin
 `storage` | [StorageConfig](#storageconfig)<br>Configuration for the storage plugin
 `tools` | [ToolsConfig](#toolsconfig)<br>Configuration for the tools plugin
@@ -317,14 +317,6 @@ Name | Description
 `description` | `string`<br>Human-readable description
 `computeId` | `string`<br>Compute ID for this network (used when allocating compute resources)
 
-## ExamplePluginConfig
-
-Configuration for the ExamplePlugin plugin
-
-Name | Description
----- | -----------
-`greetingMessage` | `string`<br>The Greeting provided by the example plugin
-
 ## BuildConfig
 
 Configuration for the build plugin
@@ -335,6 +327,7 @@ Name | Description
 `devices` | [DeviceConfig](#deviceconfig)<br>Device configuration
 `maxConformCount` | `integer`<br>Maximum number of conforms to run at once
 `agentShutdownIfDisabledGracePeriod` | `string`<br>Time to wait before shutting down an agent that has been disabled Used if no value is set on the actual pool.
+`robomerge` | [RoboMergeConfig](#robomergeconfig)<br>RoboMerge branchmap file configuration for the merge graph service. When null or empty, the service is inert (no Perforce connections, no API data).
 `artifactTypes` | [ArtifactTypeConfig](#artifacttypeconfig)`[]`<br>Configuration for different artifact types
 `projects` | [ProjectConfig](Projects.md)`[]`<br>List of projects
 `enableConformTasks` | `boolean`<br>Whether to allow conform tasks to run
@@ -447,6 +440,16 @@ Name | Description
 `Automation` | Available to CIS jobs
 `Shared` | Shared by users with remote checking and checkouts
 
+## RoboMergeConfig
+
+Configuration for the RoboMerge branchmap service
+
+Name | Description
+---- | -----------
+`branchmapFiles` | `string[]`<br>Perforce depot paths to branchmap JSON files. Plain depot paths (e.g. "//MyProject/Main/Build/Streams/myproject.branchmap.json"). All files are read from the Perforce cluster specified by .
+`perforceCluster` | `string`<br>Perforce cluster name to connect to for reading branchmap files. Maps to a connection configured in BuildServerConfig.Perforce[]. Default: "default" (uses the default Perforce connection).
+`tickInterval` | `string`<br>How often to poll Perforce for branchmap file changes. Default: 5 minutes. Changing this requires a server restart (ticker interval is set at construction time).
+
 ## ArtifactTypeConfig
 
 Configuration for an artifact
@@ -460,13 +463,13 @@ Name | Description
 `keepDays` | `integer`<br>Number of days to retain artifacts of this type
 `namespaceId` | `string`<br>Storage namespace to use for this artifact types
 
-## ExperimentalConfig
+## ExamplePluginConfig
 
-Configuration for the experimental plugin
+Configuration for the ExamplePlugin plugin
 
 Name | Description
 ---- | -----------
-`notifications` | [NotificationConfig](Notification.md)`[]`<br>Notifications to be sent on the Horde server
+`greetingMessage` | `string`<br>The Greeting provided by the example plugin
 
 ## SecretsConfig
 
@@ -550,6 +553,14 @@ Name | Description
 ---- | -----------
 `vaultUri` | `string`<br>The Uri to the Azure Key Vault.
 
+## ExperimentalConfig
+
+Configuration for the experimental plugin
+
+Name | Description
+---- | -----------
+`notifications` | [NotificationConfig](Notification.md)`[]`<br>Notifications to be sent on the Horde server
+
 ## SymbolsConfig
 
 Configuration for the tools system
diff --git a/Engine/Source/Programs/Horde/Docs/Deployment/ServerSettings.md b/Engine/Source/Programs/Horde/Docs/Deployment/ServerSettings.md
--- a/Engine/Source/Programs/Horde/Docs/Deployment/ServerSettings.md
+++ b/Engine/Source/Programs/Horde/Docs/Deployment/ServerSettings.md
@@ -129,10 +129,10 @@ Name | Description
 ---- | -----------
 `analytics` | [AnalyticsServerConfig](#analyticsserverconfig)<br>Configuration for the analytics plugin
 `compute` | [ComputeServerConfig](#computeserverconfig)<br>Configuration for the compute plugin
-`exampleplugin` | [ExamplePluginServerConfig](#examplepluginserverconfig)<br>Configuration for the exampleplugin plugin
 `build` | [BuildServerConfig](#buildserverconfig)<br>Configuration for the build plugin
-`experimental` | [ExperimentalServerConfig](#experimentalserverconfig)<br>Configuration for the experimental plugin
+`exampleplugin` | [ExamplePluginServerConfig](#examplepluginserverconfig)<br>Configuration for the exampleplugin plugin
 `secrets` | [SecretsServerConfig](#secretsserverconfig)<br>Configuration for the secrets plugin
+`experimental` | [ExperimentalServerConfig](#experimentalserverconfig)<br>Configuration for the experimental plugin
 `symbols` | [PluginServerConfig](#pluginserverconfig)<br>Configuration for the symbols plugin
 `storage` | [StorageServerConfig](#storageserverconfig)<br>Configuration for the storage plugin
 `tools` | [ToolsServerConfig](#toolsserverconfig)<br>Configuration for the tools plugin
@@ -219,14 +219,6 @@ Name | Description
 `ComputeQueueAwsMetric` | A no-op strategy that reports metrics to let an external AWS auto-scaling policy scale the fleet
 `LeaseUtilizationAwsMetric` | A no-op strategy that reports metrics to let an external AWS auto-scaling policy scale the fleet
 
-## ExamplePluginServerConfig
-
-Server configuration for the ExamplePlugin
-
-Name | Description
----- | -----------
-`enabled` | `boolean`<br>Whether the plugin should be enabled or not
-
 ## BuildServerConfig
 
 Static configuration for the build plugin
@@ -238,7 +230,7 @@ Name | Description
 `perforceConnectionPoolSize` | `integer`<br>Number of pooled perforce connections to keep
 `enableConformTasks` | `boolean`<br>Whether to enable the conform task source.
 `p4SwarmUrl` | `string`<br>Url of P4 Swarm installation
-`robomergeUrl` | `string`<br>Url of Robomergem installation
+`robomergeUrl` | `string`<br>Url of Robomerge installation
 `commitsViewerUrl` | `string`<br>Url of Commits Viewer
 `jiraUsername` | `string`<br>The Jira service account user name
 `jiraApiToken` | `string`<br>The Jira service account API token
@@ -372,9 +364,9 @@ Name | Description
 `maxChildCount` | `integer`<br>Maximum number of children in each node
 `sliceThreshold` | `integer`<br>Threshold hash value for splitting interior nodes
 
-## ExperimentalServerConfig
+## ExamplePluginServerConfig
 
-Server configuration for the experimental plugin
+Server configuration for the ExamplePlugin
 
 Name | Description
 ---- | -----------
@@ -389,6 +381,14 @@ Name | Description
 `withAws` | `boolean`<br>Whether to enable Amazon Web Services (AWS) specific features
 `enabled` | `boolean`<br>Whether the plugin should be enabled or not
 
+## ExperimentalServerConfig
+
+Server configuration for the experimental plugin
+
+Name | Description
+---- | -----------
+`enabled` | `boolean`<br>Whether the plugin should be enabled or not
+
 ## PluginServerConfig
 
 Base class for plugin server config objects
