[Test Automation Hub] Fix issue with session meta match using '+' list syntax
+ make test metadata match case insensitive

#rnx
#Synced-CL 51823047

diff --git a/Engine/Source/Programs/Horde/HordeDashboard/plugins/test-automation/testData.ts b/Engine/Source/Programs/Horde/HordeDashboard/plugins/test-automation/testData.ts
--- a/Engine/Source/Programs/Horde/HordeDashboard/plugins/test-automation/testData.ts
+++ b/Engine/Source/Programs/Horde/HordeDashboard/plugins/test-automation/testData.ts
@@ -500,14 +500,15 @@ export class TestNameRef {
         return toMatch.some(tag => this.tagRefs?.values().some(r => r.isMatch(tag)));
     }
 
-     isMetaMatch(metaFilter: MetadataFilter): boolean {
+    isMetaMatch(metaFilter: MetadataFilter): boolean {
+        // match test metadata (not session metadata)
         for (const [key, values] of metaFilter.entries()) {
             const thisValue = this[key];
             if (!thisValue) {
                 // missing key
                 return false;
             }
-            const toValidate = values.some((v) => v.indexOf('+') > -1) ? [thisValue] : thisValue.split('+');
+            const toValidate = values.some((v) => v.indexOf('+') > -1) ? [thisValue.toLowerCase()] : thisValue.split('+').map((v) => v.toLowerCase());
             if (!toValidate.some((v) => values.includes(v))) {
                 return false;
             }
@@ -579,7 +580,10 @@ export class MetadataRef {
         return this.orderedKeys.filter((k) => commonKeys.includes(k)).map((k) => this.entries[k]!).filter((v) => !MetadataRef.valuesToIgnore.has(v));
     }
 
-    isMatch(metaFilter: MetadataFilter): boolean {
+    isMatch(metaFilter: MetadataFilter, extact?: boolean): boolean {
+        if (extact && Object.keys(this.entries).length !== metaFilter.size) {
+            return false;
+        }
         for (const [key, values] of metaFilter.entries()) {
             const thisValue = this.entries[key];
             if (!thisValue) {
@@ -587,7 +591,7 @@ export class MetadataRef {
                 return false;
             }
             const toValidate = values.some((v) => v.indexOf('+') > -1) ? [thisValue] : thisValue.split('+');
-            if (!toValidate.some((v) => values.includes(v))) {
+            if (!toValidate.some((v) => values.includes(v)) || (extact && values.length  !== toValidate.length)) {
                 return false;
             }
         }
@@ -926,11 +930,12 @@ export class TestStreamRef {
     getMatchingTests(toMatchNames?: string[], toMatchTags?: string[], toMatchMeta?: MetadataFilter): Set<TestNameRef> {
         const lowerToMatchNames = toMatchNames?.map((i) => i.toLowerCase());
         const lowerToMatchTags = toMatchTags?.map((i) => i.toLowerCase());
+        const lowerToMatchMeta = toMatchMeta? new Map(toMatchMeta.entries().map(([key, values]) => [key, values.map((v) => v.toLowerCase())])) : undefined;
         const matches: Set<TestNameRef> = new Set();
         this.tests.forEach((test) => {
             const isMatchNames = !lowerToMatchNames || lowerToMatchNames.some((key) => test.isMatch(key));
             const isMatchTags = !lowerToMatchTags || test.includesTags(lowerToMatchTags);
-            const isMatchMeta = !toMatchMeta || test.isMetaMatch(toMatchMeta);
+            const isMatchMeta = !lowerToMatchMeta || test.isMetaMatch(lowerToMatchMeta);
             if (isMatchNames && isMatchTags && isMatchMeta) {
                 matches.add(test);
             }
diff --git a/Engine/Source/Programs/Horde/HordeDashboard/plugins/test-automation/testDataV2Fetcher.ts b/Engine/Source/Programs/Horde/HordeDashboard/plugins/test-automation/testDataV2Fetcher.ts
--- a/Engine/Source/Programs/Horde/HordeDashboard/plugins/test-automation/testDataV2Fetcher.ts
+++ b/Engine/Source/Programs/Horde/HordeDashboard/plugins/test-automation/testDataV2Fetcher.ts
@@ -390,7 +390,7 @@ export class TestDataV2 implements TestVersionFetcher {
 
         // search for known metadata
         const metaFilter = new Map(Object.entries(testDetails.metadata).map(([key, value]) => [key, [value]]));
-        const metaRef = this.metadataRefMap.values().find((meta) => Object.keys(meta.entries).length === metaFilter.size &&  meta.isMatch(metaFilter)) ?? this.createMetadataRef({id: testData.id, entries: testDetails.metadata});
+        const metaRef = this.metadataRefMap.values().find((meta) => meta.isMatch(metaFilter, true)) ?? this.createMetadataRef({id: testData.id, entries: testDetails.metadata});
         result.meta = metaRef;
 
         // search for known test name
