Deprecate the legacy test report hooks

#rnx
#Synced-CL 51823047

diff --git a/Engine/Source/Programs/Horde/HordeDashboard/src/components/StreamSummary.tsx b/Engine/Source/Programs/Horde/HordeDashboard/src/components/StreamSummary.tsx
--- a/Engine/Source/Programs/Horde/HordeDashboard/src/components/StreamSummary.tsx
+++ b/Engine/Source/Programs/Horde/HordeDashboard/src/components/StreamSummary.tsx
@@ -15,7 +15,6 @@ import { IssueModalV2 } from './IssueViewV2';
 import { useQuery } from "horde/base/utilities/hooks";
 import { SchedulePane } from './SchedulePane';
 import { IssueStatusIconV2 } from './StatusIcon';
-import { BuildHealthTestReportPanel } from './TestReportPanel';
 import { getHordeStyling } from '../styles/Styles';
 
 
@@ -835,7 +834,6 @@ const HealthPanelIssues: React.FC<{ desktopAlerts?: boolean }> = observer(({ des
             </Stack>
          </div>}
       </Stack>
-      <BuildHealthTestReportPanel streamId={handler.streamId!} />
    </Stack>
 
 });
diff --git a/Engine/Source/Programs/Horde/HordeDashboard/src/components/TestReportPanel.tsx b/Engine/Source/Programs/Horde/HordeDashboard/src/components/TestReportPanel.tsx
deleted file mode 100644
diff --git a/Engine/Source/Programs/Horde/HordeDashboard/src/components/jobDetailsV2/StepDetailTestPanel.tsx b/Engine/Source/Programs/Horde/HordeDashboard/src/components/jobDetailsV2/StepDetailTestPanel.tsx
--- a/Engine/Source/Programs/Horde/HordeDashboard/src/components/jobDetailsV2/StepDetailTestPanel.tsx
+++ b/Engine/Source/Programs/Horde/HordeDashboard/src/components/jobDetailsV2/StepDetailTestPanel.tsx
@@ -162,12 +162,14 @@ export const StepTestReportPanel: React.FC<{ jobDetails: JobDetailsV2, stepId?:
       return null;
    }
 
+   const isPreflight = jobDetails.jobData?.preflightCommitId !== undefined;
 
    return (<Stack id={sideRail.url} styles={{ root: { paddingTop: 18, paddingRight: 12 } }}>
       <Stack className={hordeClasses.raised}>
          <Stack tokens={{ childrenGap: 12 }}>
             <Text variant="mediumPlus" styles={{ root: { fontFamily: "Horde Open Sans SemiBold" } }}>Test Report</Text>
             <Stack styles={{ root: { paddingLeft: 4, paddingRight: 0, paddingTop: 8, paddingBottom: 4 } }}>
+               {!isPreflight && <Link className={"view-log-link"} to={`/automation?stream=${jobDetails.stream?.id}`}>Open Automation Hub</Link>}
                {items.map((test) => <Stack key={test.item.id} className={styles.item} horizontal wrap tokens={{ childrenGap: 30 }}>
                   <Link className={"view-log-link"} to={`/testreport/${test.item.id}`}>{test.name}</Link>
                   <ComponentMount component={test.linkComponent} id={test.item.id} fullName={test.item.key} />
