Mark Test automation hub V1 as legacy

#rnx
#Synced-CL 51823047

diff --git a/Engine/Source/Programs/Horde/HordeDashboard/plugins/test-automation/index.tsx b/Engine/Source/Programs/Horde/HordeDashboard/plugins/test-automation/index.tsx
--- a/Engine/Source/Programs/Horde/HordeDashboard/plugins/test-automation/index.tsx
+++ b/Engine/Source/Programs/Horde/HordeDashboard/plugins/test-automation/index.tsx
@@ -13,7 +13,7 @@ registerHordePlugin({
    mount: {
       type: MountType.TopNav,
       context: "Tools",
-      text: "Test Automation Hub(v2)",
+      text: "Test Automation Hub",
       route: `/test-automation`
    }
 })
\ No newline at end of file
diff --git a/Engine/Source/Programs/Horde/HordeDashboard/src/backend/Api.ts b/Engine/Source/Programs/Horde/HordeDashboard/src/backend/Api.ts
--- a/Engine/Source/Programs/Horde/HordeDashboard/src/backend/Api.ts
+++ b/Engine/Source/Programs/Horde/HordeDashboard/src/backend/Api.ts
@@ -3524,7 +3524,7 @@ export type GetDashboardFeaturesResponse = {
 	/** Show the device manager on the server menu */
 	showDeviceManager?: boolean;
 
-	/** Show automated tests on the server menu */
+	/** Show automated tests (legacy) on the server menu */
 	showTests?: boolean;
 
 	/** Whether to show accounts on the server menu*/
diff --git a/Engine/Source/Programs/Horde/HordeDashboard/src/components/TopNav.tsx b/Engine/Source/Programs/Horde/HordeDashboard/src/components/TopNav.tsx
--- a/Engine/Source/Programs/Horde/HordeDashboard/src/components/TopNav.tsx
+++ b/Engine/Source/Programs/Horde/HordeDashboard/src/components/TopNav.tsx
@@ -483,7 +483,7 @@ export const TopNav: React.FC<{ suppressServer?: boolean }> = observer(({ suppre
       if (features?.showTests !== false) {
          serviceItems.push({
             key: "admin_tests",
-            text: "Automation Hub",
+            text: "Automation Hub (legacy)",
             link: `/automation`
          });
       }
