Mark Test automation hub V1 as legacy (server side)

#rnx

#Synced-CL 51823047

diff --git a/Engine/Source/Programs/Horde/HordeServer/Dashboard/DashboardConfig.cs b/Engine/Source/Programs/Horde/HordeServer/Dashboard/DashboardConfig.cs
--- a/Engine/Source/Programs/Horde/HordeServer/Dashboard/DashboardConfig.cs
+++ b/Engine/Source/Programs/Horde/HordeServer/Dashboard/DashboardConfig.cs
@@ -55,9 +55,9 @@ namespace HordeServer.Dashboard
 		public bool ShowDeviceManager { get; set; } = true;
 
 		/// <summary>
-		/// Show automated tests on the server menu
+		/// Show automated tests (legacy) on the server menu
 		/// </summary>
-		public bool ShowTests { get; set; } = true;
+		public bool ShowTests { get; set; } = false;
 
 		/// <summary>
 		/// Configuration for different agent pages
diff --git a/Engine/Source/Programs/Shared/EpicGames.Horde/Users/UserMessages.cs b/Engine/Source/Programs/Shared/EpicGames.Horde/Users/UserMessages.cs
--- a/Engine/Source/Programs/Shared/EpicGames.Horde/Users/UserMessages.cs
+++ b/Engine/Source/Programs/Shared/EpicGames.Horde/Users/UserMessages.cs
@@ -257,7 +257,7 @@ namespace EpicGames.Horde.Users
 		public bool ShowDeviceManager { get; set; }
 
 		/// <summary>
-		/// Show automated tests on the server menu
+		/// Show automated tests (legacy) on the server menu
 		/// </summary>
 		public bool ShowTests { get; set; }
 
