Sending step messages to databricks.

#Synced-CL 51836478

diff --git a/Engine/Source/Programs/Horde/Plugins/Build/HordeServer.Build/Jobs/JobRpcService.cs b/Engine/Source/Programs/Horde/Plugins/Build/HordeServer.Build/Jobs/JobRpcService.cs
--- a/Engine/Source/Programs/Horde/Plugins/Build/HordeServer.Build/Jobs/JobRpcService.cs
+++ b/Engine/Source/Programs/Horde/Plugins/Build/HordeServer.Build/Jobs/JobRpcService.cs
@@ -1006,16 +1006,16 @@ namespace HordeServer.Jobs
 
 				if (!String.IsNullOrEmpty(newReport.Message))
 				{
-					List<string> reportMetaData = new List<string>() { $"ReportStepMessage={newReport.Message}"};
+					List<string> reportMetaData = new List<string>() { $"{JobReportMetadataKeys.MessageKey}={newReport.Message}"};
 
 					if (!String.IsNullOrEmpty(newReport.Severity))
 					{
-						reportMetaData.Add($"ReportStepSeverity={newReport.Severity}");
+						reportMetaData.Add($"{JobReportMetadataKeys.MessageSeverityKey}={newReport.Severity}");
 					}
 
 					if (!String.IsNullOrEmpty(newReport.Color))
 					{
-						reportMetaData.Add($"ReportStepColor={newReport.Color}");
+						reportMetaData.Add($"{JobReportMetadataKeys.ColorKey}={newReport.Color}");
 					}
 
 					Dictionary<JobStepId, List<string>> stepMetaData = new Dictionary<JobStepId, List<string>>() { { stepId, reportMetaData } };					
diff --git a/Engine/Source/Programs/Horde/Plugins/Build/HordeServer.Build/Jobs/JobStepRefCollection.cs b/Engine/Source/Programs/Horde/Plugins/Build/HordeServer.Build/Jobs/JobStepRefCollection.cs
--- a/Engine/Source/Programs/Horde/Plugins/Build/HordeServer.Build/Jobs/JobStepRefCollection.cs
+++ b/Engine/Source/Programs/Horde/Plugins/Build/HordeServer.Build/Jobs/JobStepRefCollection.cs
@@ -181,6 +181,9 @@ namespace HordeServer.Jobs
 			{
 				if (_buildConfig.CurrentValue.TryGetStream(streamId, out StreamConfig? streamConfig) && !streamConfig.TelemetryStoreId.IsEmpty)
 				{
+					string? message = stepMetadata?.FirstOrDefault(m => m.StartsWith(JobReportMetadataKeys.MessageKey + "=", StringComparison.Ordinal))?.Substring(JobReportMetadataKeys.MessageKey.Length + 1);
+					string? messageSeverity = stepMetadata?.FirstOrDefault(m => m.StartsWith(JobReportMetadataKeys.MessageSeverityKey + "=", StringComparison.Ordinal))?.Substring(JobReportMetadataKeys.MessageSeverityKey.Length + 1);
+					
 					_telemetryWriter.WriteEvent(streamConfig.TelemetryStoreId, new JobStepRefTelemetry(
 						id.ToString(),
 						id.JobId.ToString(),
@@ -201,7 +204,9 @@ namespace HordeServer.Jobs
 						streamId,
 						templateId,
 						updateIssues,
-						(finishTimeUtc != null) ? (double?)(finishTimeUtc.Value - startTimeUtc).TotalSeconds : null)
+						(finishTimeUtc != null) ? (double?)(finishTimeUtc.Value - startTimeUtc).TotalSeconds : null,
+						message,
+						messageSeverity)
 					);
 				}
 			}
diff --git a/Engine/Source/Programs/Horde/Plugins/Build/HordeServer.Build/Jobs/JobTelemetry.cs b/Engine/Source/Programs/Horde/Plugins/Build/HordeServer.Build/Jobs/JobTelemetry.cs
--- a/Engine/Source/Programs/Horde/Plugins/Build/HordeServer.Build/Jobs/JobTelemetry.cs
+++ b/Engine/Source/Programs/Horde/Plugins/Build/HordeServer.Build/Jobs/JobTelemetry.cs
@@ -617,6 +617,16 @@ namespace HordeServer.Jobs
 		/// </summary>
 		public double? Duration { get; init; }
 
+		/// <summary>
+		/// The message from the step report, if any.
+		/// </summary>
+		public string? Message { get; init; }
+
+		/// <summary>
+		/// The severity from the step report message, if any.
+		/// </summary>
+		public string? MessageSeverity { get; init; }
+
 		/// <summary>
 		/// The default event name.
 		/// </summary>
@@ -650,7 +660,9 @@ namespace HordeServer.Jobs
 		/// <param name="templateId">The template id of the owning job.</param>
 		/// <param name="updateIssues">Whether issues have been updated or not.</param>
 		/// <param name="duration">The duration of the step.</param>
-		public JobStepRefTelemetry(string id, string jobId, string batchId, string stepId, AgentId? agentId, float batchInitTime, float batchWaitTime, string? change, DateTime? finishTime, string jobName, DateTime? jobStartTime, string stepName, JobStepState? state, JobStepOutcome? outcome, PoolId? poolId, DateTime? startTime, StreamId streamId, TemplateId templateId, bool updateIssues, double? duration)
+		/// <param name="message">The message from the step report, if any.</param>
+		/// <param name="messageSeverity">The severity from the step report message, if any.</param>
+		public JobStepRefTelemetry(string id, string jobId, string batchId, string stepId, AgentId? agentId, float batchInitTime, float batchWaitTime, string? change, DateTime? finishTime, string jobName, DateTime? jobStartTime, string stepName, JobStepState? state, JobStepOutcome? outcome, PoolId? poolId, DateTime? startTime, StreamId streamId, TemplateId templateId, bool updateIssues, double? duration, string? message = null, string? messageSeverity = null)
 		{
 			EventName = DefaultEventName;
 			SchemaVersion = CurrentSchemaVersion;
@@ -674,6 +686,8 @@ namespace HordeServer.Jobs
 			AgentId = agentId;
 			UpdateIssues = updateIssues;
 			Duration = duration;
+			Message = message;
+			MessageSeverity = messageSeverity;
 		}
 
 		/// <summary>
@@ -735,4 +749,25 @@ namespace HordeServer.Jobs
 			return telemetry;
 		}
 	}
+
+	/// <summary>
+	/// Keys used to store job report data in step metadata.
+	/// </summary>
+	static class JobReportMetadataKeys
+	{
+		/// <summary>
+		/// Key for the report message.
+		/// </summary>
+		public const string MessageKey = "ReportStepMessage";
+
+		/// <summary>
+		/// Key for the report severity.
+		/// </summary>
+		public const string MessageSeverityKey = "ReportStepSeverity";
+
+		/// <summary>
+		/// Key for the report color.
+		/// </summary>
+		public const string ColorKey = "ReportStepColor";
+	}
 }
\ No newline at end of file
