Horde: Updating schema documentation

#rnx

diff --git a/Engine/Source/Programs/Horde/Docs/Config/Schema/AclActions.md b/Engine/Source/Programs/Horde/Docs/Config/Schema/AclActions.md
--- a/Engine/Source/Programs/Horde/Docs/Config/Schema/AclActions.md
+++ b/Engine/Source/Programs/Horde/Docs/Config/Schema/AclActions.md
@@ -63,44 +63,6 @@
 | `DeleteAccount` |  |
 | `ViewAccount` |  |
 
-## Agents
-
-| Name | Description |
-| ---- | ----------- |
-| `CreateAgent` | Ability to create an agent. This may be done explicitly, or granted to agents to allow them to self-register. |
-| `CreateWorkstationAgent` | Ability to create a workstation agent. Granted to users with hardware to offer. |
-| `UpdateAgent` | Update an agent's name, pools, etc... |
-| `DeleteAgent` | Soft-delete an agent |
-| `ViewAgent` | View an agent |
-| `ListAgents` | List the available agents |
-
-## Compute
-
-| Name | Description |
-| ---- | ----------- |
-| `AddComputeTasks` | User can add tasks to the compute cluster |
-| `GetComputeTasks` | User can get and list tasks from the compute cluster |
-| `UbaCacheRead` | User can read from Unreal Build Accelerator cache |
-| `UbaCacheWrite` | User can write to Unreal Build Accelerator cache |
-
-## Leases
-
-| Name | Description |
-| ---- | ----------- |
-| `ViewLeases` | View all the leases that an agent has worked on |
-| `ViewLeaseTasks` | View the task data for a lease |
-
-## Logs
-
-| Name | Description |
-| ---- | ----------- |
-| `CreateLog` | Ability to create a log. Implicitly granted to agents. |
-| `UpdateLog` | Ability to update log metadata |
-| `ViewLog` | Ability to view a log contents |
-| `WriteLogData` | Ability to write log data |
-| `CreateEvent` | Ability to create events |
-| `ViewEvent` | Ability to view events |
-
 ## Notices
 
 | Name | Description |
@@ -109,16 +71,6 @@
 | `UpdateNotice` |  |
 | `DeleteNotice` |  |
 
-## Pools
-
-| Name | Description |
-| ---- | ----------- |
-| `CreatePool` | Create a global pool of agents |
-| `UpdatePool` | Modify an agent pool |
-| `DeletePool` | Delete an agent pool |
-| `ViewPool` | Ability to view a pool |
-| `ListPools` | View all the available agent pools |
-
 ## ServiceAccounts
 
 | Name | Description |
@@ -128,21 +80,6 @@
 | `DeleteAccount` |  |
 | `ViewAccount` |  |
 
-## Sessions
-
-| Name | Description |
-| ---- | ----------- |
-| `CreateSession` | Granted to agents to call CreateSession, which returns a bearer token identifying themselves valid to call UpdateSesssion via gRPC. |
-| `ViewSession` | Allows viewing information about an agent session |
-
-## Software
-
-| Name | Description |
-| ---- | ----------- |
-| `UploadSoftware` | Ability to upload new versions of the agent software |
-| `DownloadSoftware` | Ability to download the agent software |
-| `DeleteSoftware` | Ability to delete agent software |
-
 ## Accounts
 
 | Name | Description |
@@ -152,86 +89,6 @@
 | `DeleteAccount` |  |
 | `ViewAccount` |  |
 
-## Agents
-
-| Name | Description |
-| ---- | ----------- |
-| `CreateAgent` |  |
-| `CreateWorkstationAgent` |  |
-| `UpdateAgent` |  |
-| `DeleteAgent` |  |
-| `ViewAgent` |  |
-| `ListAgents` |  |
-
-## Artifacts
-
-| Name | Description |
-| ---- | ----------- |
-| `ReadArtifact` | Permission to read from an artifact |
-| `WriteArtifact` | Permission to write to an artifact |
-| `DeleteArtifact` | Permission to delete to an artifact |
-| `UploadArtifact` | Ability to create an artifact. Typically just for debugging; agents have this access for a particular session. |
-| `DownloadArtifact` | Ability to download an artifact |
-
-## Bisect
-
-| Name | Description |
-| ---- | ----------- |
-| `CreateBisectTask` | Ability to start new bisect tasks |
-| `UpdateBisectTask` | Ability to update a bisect task |
-| `ViewBisectTask` | Ability to view a bisect task |
-
-## BuildHealth
-
-| Name | Description |
-| ---- | ----------- |
-| `AddUpdateDeleteBuildHealthFilter` | Ability to add, update, and delete build health filters. |
-
-## Compute
-
-| Name | Description |
-| ---- | ----------- |
-| `AddComputeTasks` |  |
-| `GetComputeTasks` |  |
-| `UbaCacheRead` |  |
-| `UbaCacheWrite` |  |
-
-## Devices
-
-| Name | Description |
-| ---- | ----------- |
-| `DeviceRead` | Ability to read devices |
-| `DeviceWrite` | Ability to write devices |
-
-## Jobs
-
-| Name | Description |
-| ---- | ----------- |
-| `CreateJob` | Ability to start new jobs |
-| `UpdateJob` | Rename a job, modify its priority, etc... |
-| `DeleteJob` | Delete a job properties |
-| `ExecuteJob` | Allows updating a job metadata (name, changelist number, step properties, new groups, job states, etc...). Typically granted to agents. Not user facing. |
-| `RetryJobStep` | Ability to retry a failed job step |
-| `ViewJob` | Ability to view a job |
-
-## Leases
-
-| Name | Description |
-| ---- | ----------- |
-| `ViewLeases` |  |
-| `ViewLeaseTasks` |  |
-
-## Logs
-
-| Name | Description |
-| ---- | ----------- |
-| `CreateLog` |  |
-| `UpdateLog` |  |
-| `ViewLog` |  |
-| `WriteLogData` |  |
-| `CreateEvent` |  |
-| `ViewEvent` |  |
-
 ## Notices
 
 | Name | Description |
@@ -240,38 +97,6 @@
 | `UpdateNotice` |  |
 | `DeleteNotice` |  |
 
-## Notifications
-
-| Name | Description |
-| ---- | ----------- |
-| `CreateSubscription` | Ability to subscribe to notifications |
-
-## Pools
-
-| Name | Description |
-| ---- | ----------- |
-| `CreatePool` |  |
-| `UpdatePool` |  |
-| `DeletePool` |  |
-| `ViewPool` |  |
-| `ListPools` |  |
-
-## Projects
-
-| Name | Description |
-| ---- | ----------- |
-| `CreateProject` | Allows the creation of new projects |
-| `DeleteProject` | Allows deletion of projects. |
-| `UpdateProject` | Modify attributes of a project (name, categories, etc...) |
-| `ViewProject` | View information about a project |
-
-## Replicators
-
-| Name | Description |
-| ---- | ----------- |
-| `UpdateReplicator` | Allows deletion of projects. |
-| `ViewReplicator` | Allows the creation of new projects |
-
 ## ServiceAccounts
 
 | Name | Description |
@@ -281,38 +106,37 @@
 | `DeleteAccount` |  |
 | `ViewAccount` |  |
 
-## Sessions
+## Accounts
 
 | Name | Description |
 | ---- | ----------- |
-| `CreateSession` |  |
-| `ViewSession` |  |
+| `CreateAccount` |  |
+| `UpdateAccount` |  |
+| `DeleteAccount` |  |
+| `ViewAccount` |  |
 
-## Software
+## Notices
 
 | Name | Description |
 | ---- | ----------- |
-| `UploadSoftware` |  |
-| `DownloadSoftware` |  |
-| `DeleteSoftware` |  |
+| `CreateNotice` |  |
+| `UpdateNotice` |  |
+| `DeleteNotice` |  |
 
-## Streams
+## Secrets
 
 | Name | Description |
 | ---- | ----------- |
-| `CreateStream` | Allows the creation of new streams within a project |
-| `UpdateStream` | Allows updating a stream (agent types, templates, schedules) |
-| `DeleteStream` | Allows deleting a stream |
-| `ViewStream` | Ability to view a stream |
-| `ViewChanges` | View changes submitted to a stream. NOTE: this returns responses from the server's Perforce account, which may be a priviledged user. |
-| `ViewTemplate` | View template associated with a stream |
+| `ViewSecret` | View a credential |
 
-## TestData
+## ServiceAccounts
 
 | Name | Description |
 | ---- | ----------- |
-| `TestHealthRead` | Ability to read test health |
-| `TestHealthWrite` | Ability to write test health |
+| `CreateAccount` |  |
+| `UpdateAccount` |  |
+| `DeleteAccount` |  |
+| `ViewAccount` |  |
 
 ## Accounts
 
@@ -323,86 +147,6 @@
 | `DeleteAccount` |  |
 | `ViewAccount` |  |
 
-## Agents
-
-| Name | Description |
-| ---- | ----------- |
-| `CreateAgent` |  |
-| `CreateWorkstationAgent` |  |
-| `UpdateAgent` |  |
-| `DeleteAgent` |  |
-| `ViewAgent` |  |
-| `ListAgents` |  |
-
-## Artifacts
-
-| Name | Description |
-| ---- | ----------- |
-| `ReadArtifact` |  |
-| `WriteArtifact` |  |
-| `DeleteArtifact` |  |
-| `UploadArtifact` |  |
-| `DownloadArtifact` |  |
-
-## Bisect
-
-| Name | Description |
-| ---- | ----------- |
-| `CreateBisectTask` |  |
-| `UpdateBisectTask` |  |
-| `ViewBisectTask` |  |
-
-## BuildHealth
-
-| Name | Description |
-| ---- | ----------- |
-| `AddUpdateDeleteBuildHealthFilter` |  |
-
-## Compute
-
-| Name | Description |
-| ---- | ----------- |
-| `AddComputeTasks` |  |
-| `GetComputeTasks` |  |
-| `UbaCacheRead` |  |
-| `UbaCacheWrite` |  |
-
-## Devices
-
-| Name | Description |
-| ---- | ----------- |
-| `DeviceRead` |  |
-| `DeviceWrite` |  |
-
-## Jobs
-
-| Name | Description |
-| ---- | ----------- |
-| `CreateJob` |  |
-| `UpdateJob` |  |
-| `DeleteJob` |  |
-| `ExecuteJob` |  |
-| `RetryJobStep` |  |
-| `ViewJob` |  |
-
-## Leases
-
-| Name | Description |
-| ---- | ----------- |
-| `ViewLeases` |  |
-| `ViewLeaseTasks` |  |
-
-## Logs
-
-| Name | Description |
-| ---- | ----------- |
-| `CreateLog` |  |
-| `UpdateLog` |  |
-| `ViewLog` |  |
-| `WriteLogData` |  |
-| `CreateEvent` |  |
-| `ViewEvent` |  |
-
 ## Notices
 
 | Name | Description |
@@ -411,390 +155,10 @@
 | `UpdateNotice` |  |
 | `DeleteNotice` |  |
 
-## Notifications
-
-| Name | Description |
-| ---- | ----------- |
-| `CreateSubscription` |  |
-
-## Pools
-
-| Name | Description |
-| ---- | ----------- |
-| `CreatePool` |  |
-| `UpdatePool` |  |
-| `DeletePool` |  |
-| `ViewPool` |  |
-| `ListPools` |  |
-
-## Projects
-
-| Name | Description |
-| ---- | ----------- |
-| `CreateProject` |  |
-| `DeleteProject` |  |
-| `UpdateProject` |  |
-| `ViewProject` |  |
-
-## Replicators
-
-| Name | Description |
-| ---- | ----------- |
-| `UpdateReplicator` |  |
-| `ViewReplicator` |  |
-
-## ServiceAccounts
-
-| Name | Description |
-| ---- | ----------- |
-| `CreateAccount` |  |
-| `UpdateAccount` |  |
-| `DeleteAccount` |  |
-| `ViewAccount` |  |
-
-## Sessions
-
-| Name | Description |
-| ---- | ----------- |
-| `CreateSession` |  |
-| `ViewSession` |  |
-
-## Software
-
-| Name | Description |
-| ---- | ----------- |
-| `UploadSoftware` |  |
-| `DownloadSoftware` |  |
-| `DeleteSoftware` |  |
-
-## Streams
-
-| Name | Description |
-| ---- | ----------- |
-| `CreateStream` |  |
-| `UpdateStream` |  |
-| `DeleteStream` |  |
-| `ViewStream` |  |
-| `ViewChanges` |  |
-| `ViewTemplate` |  |
-
-## TestData
-
-| Name | Description |
-| ---- | ----------- |
-| `TestHealthRead` |  |
-| `TestHealthWrite` |  |
-
-## Accounts
-
-| Name | Description |
-| ---- | ----------- |
-| `CreateAccount` |  |
-| `UpdateAccount` |  |
-| `DeleteAccount` |  |
-| `ViewAccount` |  |
-
-## Agents
-
-| Name | Description |
-| ---- | ----------- |
-| `CreateAgent` |  |
-| `CreateWorkstationAgent` |  |
-| `UpdateAgent` |  |
-| `DeleteAgent` |  |
-| `ViewAgent` |  |
-| `ListAgents` |  |
-
-## Artifacts
-
-| Name | Description |
-| ---- | ----------- |
-| `ReadArtifact` |  |
-| `WriteArtifact` |  |
-| `DeleteArtifact` |  |
-| `UploadArtifact` |  |
-| `DownloadArtifact` |  |
-
-## Bisect
-
-| Name | Description |
-| ---- | ----------- |
-| `CreateBisectTask` |  |
-| `UpdateBisectTask` |  |
-| `ViewBisectTask` |  |
-
-## BuildHealth
-
-| Name | Description |
-| ---- | ----------- |
-| `AddUpdateDeleteBuildHealthFilter` |  |
-
-## Compute
-
-| Name | Description |
-| ---- | ----------- |
-| `AddComputeTasks` |  |
-| `GetComputeTasks` |  |
-| `UbaCacheRead` |  |
-| `UbaCacheWrite` |  |
-
-## Devices
-
-| Name | Description |
-| ---- | ----------- |
-| `DeviceRead` |  |
-| `DeviceWrite` |  |
-
-## Jobs
-
-| Name | Description |
-| ---- | ----------- |
-| `CreateJob` |  |
-| `UpdateJob` |  |
-| `DeleteJob` |  |
-| `ExecuteJob` |  |
-| `RetryJobStep` |  |
-| `ViewJob` |  |
-
-## Leases
-
-| Name | Description |
-| ---- | ----------- |
-| `ViewLeases` |  |
-| `ViewLeaseTasks` |  |
-
-## Logs
-
-| Name | Description |
-| ---- | ----------- |
-| `CreateLog` |  |
-| `UpdateLog` |  |
-| `ViewLog` |  |
-| `WriteLogData` |  |
-| `CreateEvent` |  |
-| `ViewEvent` |  |
-
-## Notices
-
-| Name | Description |
-| ---- | ----------- |
-| `CreateNotice` |  |
-| `UpdateNotice` |  |
-| `DeleteNotice` |  |
-
-## Notifications
-
-| Name | Description |
-| ---- | ----------- |
-| `CreateSubscription` |  |
-
-## Pools
-
-| Name | Description |
-| ---- | ----------- |
-| `CreatePool` |  |
-| `UpdatePool` |  |
-| `DeletePool` |  |
-| `ViewPool` |  |
-| `ListPools` |  |
-
-## Projects
-
-| Name | Description |
-| ---- | ----------- |
-| `CreateProject` |  |
-| `DeleteProject` |  |
-| `UpdateProject` |  |
-| `ViewProject` |  |
-
-## Replicators
-
-| Name | Description |
-| ---- | ----------- |
-| `UpdateReplicator` |  |
-| `ViewReplicator` |  |
-
 ## Secrets
 
 | Name | Description |
 | ---- | ----------- |
-| `ViewSecret` | View a credential |
-
-## ServiceAccounts
-
-| Name | Description |
-| ---- | ----------- |
-| `CreateAccount` |  |
-| `UpdateAccount` |  |
-| `DeleteAccount` |  |
-| `ViewAccount` |  |
-
-## Sessions
-
-| Name | Description |
-| ---- | ----------- |
-| `CreateSession` |  |
-| `ViewSession` |  |
-
-## Software
-
-| Name | Description |
-| ---- | ----------- |
-| `UploadSoftware` |  |
-| `DownloadSoftware` |  |
-| `DeleteSoftware` |  |
-
-## Streams
-
-| Name | Description |
-| ---- | ----------- |
-| `CreateStream` |  |
-| `UpdateStream` |  |
-| `DeleteStream` |  |
-| `ViewStream` |  |
-| `ViewChanges` |  |
-| `ViewTemplate` |  |
-
-## TestData
-
-| Name | Description |
-| ---- | ----------- |
-| `TestHealthRead` |  |
-| `TestHealthWrite` |  |
-
-## Accounts
-
-| Name | Description |
-| ---- | ----------- |
-| `CreateAccount` |  |
-| `UpdateAccount` |  |
-| `DeleteAccount` |  |
-| `ViewAccount` |  |
-
-## Agents
-
-| Name | Description |
-| ---- | ----------- |
-| `CreateAgent` |  |
-| `CreateWorkstationAgent` |  |
-| `UpdateAgent` |  |
-| `DeleteAgent` |  |
-| `ViewAgent` |  |
-| `ListAgents` |  |
-
-## Artifacts
-
-| Name | Description |
-| ---- | ----------- |
-| `ReadArtifact` |  |
-| `WriteArtifact` |  |
-| `DeleteArtifact` |  |
-| `UploadArtifact` |  |
-| `DownloadArtifact` |  |
-
-## Bisect
-
-| Name | Description |
-| ---- | ----------- |
-| `CreateBisectTask` |  |
-| `UpdateBisectTask` |  |
-| `ViewBisectTask` |  |
-
-## BuildHealth
-
-| Name | Description |
-| ---- | ----------- |
-| `AddUpdateDeleteBuildHealthFilter` |  |
-
-## Compute
-
-| Name | Description |
-| ---- | ----------- |
-| `AddComputeTasks` |  |
-| `GetComputeTasks` |  |
-| `UbaCacheRead` |  |
-| `UbaCacheWrite` |  |
-
-## Devices
-
-| Name | Description |
-| ---- | ----------- |
-| `DeviceRead` |  |
-| `DeviceWrite` |  |
-
-## Jobs
-
-| Name | Description |
-| ---- | ----------- |
-| `CreateJob` |  |
-| `UpdateJob` |  |
-| `DeleteJob` |  |
-| `ExecuteJob` |  |
-| `RetryJobStep` |  |
-| `ViewJob` |  |
-
-## Leases
-
-| Name | Description |
-| ---- | ----------- |
-| `ViewLeases` |  |
-| `ViewLeaseTasks` |  |
-
-## Logs
-
-| Name | Description |
-| ---- | ----------- |
-| `CreateLog` |  |
-| `UpdateLog` |  |
-| `ViewLog` |  |
-| `WriteLogData` |  |
-| `CreateEvent` |  |
-| `ViewEvent` |  |
-
-## Notices
-
-| Name | Description |
-| ---- | ----------- |
-| `CreateNotice` |  |
-| `UpdateNotice` |  |
-| `DeleteNotice` |  |
-
-## Notifications
-
-| Name | Description |
-| ---- | ----------- |
-| `CreateSubscription` |  |
-
-## Pools
-
-| Name | Description |
-| ---- | ----------- |
-| `CreatePool` |  |
-| `UpdatePool` |  |
-| `DeletePool` |  |
-| `ViewPool` |  |
-| `ListPools` |  |
-
-## Projects
-
-| Name | Description |
-| ---- | ----------- |
-| `CreateProject` |  |
-| `DeleteProject` |  |
-| `UpdateProject` |  |
-| `ViewProject` |  |
-
-## Replicators
-
-| Name | Description |
-| ---- | ----------- |
-| `UpdateReplicator` |  |
-| `ViewReplicator` |  |
-
-## Secrets
-
-| Name | Description |
-| ---- | ----------- |
 | `ViewSecret` |  |
 
 ## ServiceAccounts
@@ -806,38 +170,11 @@
 | `DeleteAccount` |  |
 | `ViewAccount` |  |
 
-## Sessions
+## Symbols
 
 | Name | Description |
 | ---- | ----------- |
-| `CreateSession` |  |
-| `ViewSession` |  |
-
-## Software
-
-| Name | Description |
-| ---- | ----------- |
-| `UploadSoftware` |  |
-| `DownloadSoftware` |  |
-| `DeleteSoftware` |  |
-
-## Streams
-
-| Name | Description |
-| ---- | ----------- |
-| `CreateStream` |  |
-| `UpdateStream` |  |
-| `DeleteStream` |  |
-| `ViewStream` |  |
-| `ViewChanges` |  |
-| `ViewTemplate` |  |
-
-## TestData
-
-| Name | Description |
-| ---- | ----------- |
-| `TestHealthRead` |  |
-| `TestHealthWrite` |  |
+| `ReadSymbols` | Ability to download symbols |
 
 ## Accounts
 
@@ -848,85 +185,92 @@
 | `DeleteAccount` |  |
 | `ViewAccount` |  |
 
-## Agents
+## Notices
 
 | Name | Description |
 | ---- | ----------- |
-| `CreateAgent` |  |
-| `CreateWorkstationAgent` |  |
-| `UpdateAgent` |  |
-| `DeleteAgent` |  |
-| `ViewAgent` |  |
-| `ListAgents` |  |
+| `CreateNotice` |  |
+| `UpdateNotice` |  |
+| `DeleteNotice` |  |
 
+## Secrets
+
+| Name | Description |
+| ---- | ----------- |
+| `ViewSecret` |  |
+
+## ServiceAccounts
+
+| Name | Description |
+| ---- | ----------- |
+| `CreateAccount` |  |
+| `UpdateAccount` |  |
+| `DeleteAccount` |  |
+| `ViewAccount` |  |
+
+## Symbols
+
+| Name | Description |
+| ---- | ----------- |
+| `ReadSymbols` |  |
+
+## Tools
+
+| Name | Description |
+| ---- | ----------- |
+| `DownloadTool` | Ability to download a tool |
+| `UploadTool` | Ability to upload new tool versions |
+
+## Accounts
+
+| Name | Description |
+| ---- | ----------- |
+| `CreateAccount` |  |
+| `UpdateAccount` |  |
+| `DeleteAccount` |  |
+| `ViewAccount` |  |
+
 ## Artifacts
 
 | Name | Description |
 | ---- | ----------- |
-| `ReadArtifact` |  |
-| `WriteArtifact` |  |
-| `DeleteArtifact` |  |
-| `UploadArtifact` |  |
-| `DownloadArtifact` |  |
+| `ReadArtifact` | Permission to read from an artifact |
+| `WriteArtifact` | Permission to write to an artifact |
+| `DeleteArtifact` | Permission to delete to an artifact |
+| `UploadArtifact` | Ability to create an artifact. Typically just for debugging; agents have this access for a particular session. |
+| `DownloadArtifact` | Ability to download an artifact |
 
 ## Bisect
 
 | Name | Description |
 | ---- | ----------- |
-| `CreateBisectTask` |  |
-| `UpdateBisectTask` |  |
-| `ViewBisectTask` |  |
+| `CreateBisectTask` | Ability to start new bisect tasks |
+| `UpdateBisectTask` | Ability to update a bisect task |
+| `ViewBisectTask` | Ability to view a bisect task |
 
 ## BuildHealth
 
 | Name | Description |
 | ---- | ----------- |
-| `AddUpdateDeleteBuildHealthFilter` |  |
+| `AddUpdateDeleteBuildHealthFilter` | Ability to add, update, and delete build health filters. |
 
-## Compute
-
-| Name | Description |
-| ---- | ----------- |
-| `AddComputeTasks` |  |
-| `GetComputeTasks` |  |
-| `UbaCacheRead` |  |
-| `UbaCacheWrite` |  |
-
 ## Devices
 
 | Name | Description |
 | ---- | ----------- |
-| `DeviceRead` |  |
-| `DeviceWrite` |  |
+| `DeviceRead` | Ability to read devices |
+| `DeviceWrite` | Ability to write devices |
 
 ## Jobs
 
 | Name | Description |
 | ---- | ----------- |
-| `CreateJob` |  |
-| `UpdateJob` |  |
-| `DeleteJob` |  |
-| `ExecuteJob` |  |
-| `RetryJobStep` |  |
-| `ViewJob` |  |
-
-## Leases
-
-| Name | Description |
-| ---- | ----------- |
-| `ViewLeases` |  |
-| `ViewLeaseTasks` |  |
-
-## Logs
-
-| Name | Description |
-| ---- | ----------- |
-| `CreateLog` |  |
-| `UpdateLog` |  |
-| `ViewLog` |  |
-| `WriteLogData` |  |
-| `CreateEvent` |  |
-| `ViewEvent` |  |
+| `CreateJob` | Ability to start new jobs |
+| `UpdateJob` | Rename a job, modify its priority, etc... |
+| `DeleteJob` | Delete a job properties |
+| `ExecuteJob` | Allows updating a job metadata (name, changelist number, step properties, new groups, job states, etc...). Typically granted to agents. Not user facing. |
+| `RetryJobStep` | Ability to retry a failed job step |
+| `ViewJob` | Ability to view a job |
 
 ## Notices
 
@@ -940,33 +284,23 @@
 
 | Name | Description |
 | ---- | ----------- |
-| `CreateSubscription` |  |
+| `CreateSubscription` | Ability to subscribe to notifications |
 
-## Pools
-
-| Name | Description |
-| ---- | ----------- |
-| `CreatePool` |  |
-| `UpdatePool` |  |
-| `DeletePool` |  |
-| `ViewPool` |  |
-| `ListPools` |  |
-
 ## Projects
 
 | Name | Description |
 | ---- | ----------- |
-| `CreateProject` |  |
-| `DeleteProject` |  |
-| `UpdateProject` |  |
-| `ViewProject` |  |
+| `CreateProject` | Allows the creation of new projects |
+| `DeleteProject` | Allows deletion of projects. |
+| `UpdateProject` | Modify attributes of a project (name, categories, etc...) |
+| `ViewProject` | View information about a project |
 
 ## Replicators
 
 | Name | Description |
 | ---- | ----------- |
-| `UpdateReplicator` |  |
-| `ViewReplicator` |  |
+| `UpdateReplicator` | Allows deletion of projects. |
+| `ViewReplicator` | Allows the creation of new projects |
 
 ## Secrets
 
@@ -983,44 +317,36 @@
 | `DeleteAccount` |  |
 | `ViewAccount` |  |
 
-## Sessions
-
-| Name | Description |
-| ---- | ----------- |
-| `CreateSession` |  |
-| `ViewSession` |  |
-
-## Software
-
-| Name | Description |
-| ---- | ----------- |
-| `UploadSoftware` |  |
-| `DownloadSoftware` |  |
-| `DeleteSoftware` |  |
-
 ## Streams
 
 | Name | Description |
 | ---- | ----------- |
-| `CreateStream` |  |
-| `UpdateStream` |  |
-| `DeleteStream` |  |
-| `ViewStream` |  |
-| `ViewChanges` |  |
-| `ViewTemplate` |  |
+| `CreateStream` | Allows the creation of new streams within a project |
+| `UpdateStream` | Allows updating a stream (agent types, templates, schedules) |
+| `DeleteStream` | Allows deleting a stream |
+| `ViewStream` | Ability to view a stream |
+| `ViewChanges` | View changes submitted to a stream. NOTE: this returns responses from the server's Perforce account, which may be a priviledged user. |
+| `ViewTemplate` | View template associated with a stream |
 
 ## Symbols
 
 | Name | Description |
 | ---- | ----------- |
-| `ReadSymbols` | Ability to download symbols |
+| `ReadSymbols` |  |
 
 ## TestData
 
 | Name | Description |
 | ---- | ----------- |
-| `TestHealthRead` |  |
-| `TestHealthWrite` |  |
+| `TestHealthRead` | Ability to read test health |
+| `TestHealthWrite` | Ability to write test health |
 
+## Tools
+
+| Name | Description |
+| ---- | ----------- |
+| `DownloadTool` |  |
+| `UploadTool` |  |
+
 ## Accounts
 
@@ -1035,12 +361,12 @@
 
 | Name | Description |
 | ---- | ----------- |
-| `CreateAgent` |  |
-| `CreateWorkstationAgent` |  |
-| `UpdateAgent` |  |
-| `DeleteAgent` |  |
-| `ViewAgent` |  |
-| `ListAgents` |  |
+| `CreateAgent` | Ability to create an agent. This may be done explicitly, or granted to agents to allow them to self-register. |
+| `CreateWorkstationAgent` | Ability to create a workstation agent. Granted to users with hardware to offer. |
+| `UpdateAgent` | Update an agent's name, pools, etc... |
+| `DeleteAgent` | Soft-delete an agent |
+| `ViewAgent` | View an agent |
+| `ListAgents` | List the available agents |
 
 ## Artifacts
 
@@ -1070,10 +396,10 @@
 
 | Name | Description |
 | ---- | ----------- |
-| `AddComputeTasks` |  |
-| `GetComputeTasks` |  |
-| `UbaCacheRead` |  |
-| `UbaCacheWrite` |  |
+| `AddComputeTasks` | User can add tasks to the compute cluster |
+| `GetComputeTasks` | User can get and list tasks from the compute cluster |
+| `UbaCacheRead` | User can read from Unreal Build Accelerator cache |
+| `UbaCacheWrite` | User can write to Unreal Build Accelerator cache |
 
 ## Devices
 
@@ -1097,19 +423,19 @@
 
 | Name | Description |
 | ---- | ----------- |
-| `ViewLeases` |  |
-| `ViewLeaseTasks` |  |
+| `ViewLeases` | View all the leases that an agent has worked on |
+| `ViewLeaseTasks` | View the task data for a lease |
 
 ## Logs
 
 | Name | Description |
 | ---- | ----------- |
-| `CreateLog` |  |
-| `UpdateLog` |  |
-| `ViewLog` |  |
-| `WriteLogData` |  |
-| `CreateEvent` |  |
-| `ViewEvent` |  |
+| `CreateLog` | Ability to create a log. Implicitly granted to agents. |
+| `UpdateLog` | Ability to update log metadata |
+| `ViewLog` | Ability to view a log contents |
+| `WriteLogData` | Ability to write log data |
+| `CreateEvent` | Ability to create events |
+| `ViewEvent` | Ability to view events |
 
 ## Notices
 
@@ -1129,11 +455,11 @@
 
 | Name | Description |
 | ---- | ----------- |
-| `CreatePool` |  |
-| `UpdatePool` |  |
-| `DeletePool` |  |
-| `ViewPool` |  |
-| `ListPools` |  |
+| `CreatePool` | Create a global pool of agents |
+| `UpdatePool` | Modify an agent pool |
+| `DeletePool` | Delete an agent pool |
+| `ViewPool` | Ability to view a pool |
+| `ListPools` | View all the available agent pools |
 
 ## Projects
 
@@ -1170,29 +496,16 @@
 
 | Name | Description |
 | ---- | ----------- |
-| `CreateSession` |  |
-| `ViewSession` |  |
+| `CreateSession` | Granted to agents to call CreateSession, which returns a bearer token identifying themselves valid to call UpdateSesssion via gRPC. |
+| `ViewSession` | Allows viewing information about an agent session |
 
 ## Software
 
 | Name | Description |
 | ---- | ----------- |
-| `UploadSoftware` |  |
-| `DownloadSoftware` |  |
-| `DeleteSoftware` |  |
+| `UploadSoftware` | Ability to upload new versions of the agent software |
+| `DownloadSoftware` | Ability to download the agent software |
+| `DeleteSoftware` | Ability to delete agent software |
 
-## Storage
-
-| Name | Description |
-| ---- | ----------- |
-| `ReadBlobs` | Ability to read blobs from the storage service |
-| `WriteBlobs` | Ability to write blobs to the storage service |
-| `ReadAliases` | Ability to read aliases from the storage service |
-| `WriteAliases` | Ability to write aliases to the storage service |
-| `DeleteAliases` | Ability to write aliases to the storage service |
-| `ReadRefs` | Ability to read refs from the storage service |
-| `WriteRefs` | Ability to write refs to the storage service |
-| `DeleteRefs` | Ability to delete refs |
-
 ## Streams
 
@@ -1218,6 +531,13 @@
 | `TestHealthRead` |  |
 | `TestHealthWrite` |  |
 
+## Tools
+
+| Name | Description |
+| ---- | ----------- |
+| `DownloadTool` |  |
+| `UploadTool` |  |
+
 ## Accounts
 
 | Name | Description |
@@ -1381,14 +701,14 @@
 
 | Name | Description |
 | ---- | ----------- |
-| `ReadBlobs` |  |
-| `WriteBlobs` |  |
-| `ReadAliases` |  |
-| `WriteAliases` |  |
-| `DeleteAliases` |  |
-| `ReadRefs` |  |
-| `WriteRefs` |  |
-| `DeleteRefs` |  |
+| `ReadBlobs` | Ability to read blobs from the storage service |
+| `WriteBlobs` | Ability to write blobs to the storage service |
+| `ReadAliases` | Ability to read aliases from the storage service |
+| `WriteAliases` | Ability to write aliases to the storage service |
+| `DeleteAliases` | Ability to write aliases to the storage service |
+| `ReadRefs` | Ability to read refs from the storage service |
+| `WriteRefs` | Ability to write refs to the storage service |
+| `DeleteRefs` | Ability to delete refs |
 
 ## Streams
 
@@ -1418,8 +738,8 @@
 
 | Name | Description |
 | ---- | ----------- |
-| `DownloadTool` | Ability to download a tool |
-| `UploadTool` | Ability to upload new tool versions |
+| `DownloadTool` |  |
+| `UploadTool` |  |
 
 ## Accounts
 
diff --git a/Engine/Source/Programs/Horde/Docs/Config/Schema/Dashboard.md b/Engine/Source/Programs/Horde/Docs/Config/Schema/Dashboard.md
--- a/Engine/Source/Programs/Horde/Docs/Config/Schema/Dashboard.md
+++ b/Engine/Source/Programs/Horde/Docs/Config/Schema/Dashboard.md
@@ -13,7 +13,7 @@ Name | Description
 `showAgentRegistration` | `boolean`<br>Whether to show the agent registration page. When using registration tokens from elsewhere this is not needed.
 `showPerforceServers` | `boolean`<br>Show the Perforce server option on the server menu
 `showDeviceManager` | `boolean`<br>Show the device manager on the server menu
-`showTests` | `boolean`<br>Show automated tests on the server menu
+`showTests` | `boolean`<br>Show automated tests (legacy) on the server menu
 `agentCategories` | [DashboardAgentCategoryConfig](#dashboardagentcategoryconfig)`[]`<br>Configuration for different agent pages
 `poolCategories` | [DashboardPoolCategoryConfig](#dashboardpoolcategoryconfig)`[]`<br>Configuration for different pool pages
 `include` | [ConfigInclude](#configinclude)`[]`<br>Includes for other configuration files
diff --git a/Engine/Source/Programs/Horde/Docs/Config/Schema/Globals.md b/Engine/Source/Programs/Horde/Docs/Config/Schema/Globals.md
--- a/Engine/Source/Programs/Horde/Docs/Config/Schema/Globals.md
+++ b/Engine/Source/Programs/Horde/Docs/Config/Schema/Globals.md
@@ -60,14 +60,14 @@
 Name | Description
 ---- | -----------
 `analytics` | [AnalyticsConfig](#analyticsconfig)<br>Configuration for the analytics plugin
-`compute` | [ComputeConfig](#computeconfig)<br>Configuration for the compute plugin
-`build` | [BuildConfig](#buildconfig)<br>Configuration for the build plugin
 `exampleplugin` | [ExamplePluginConfig](#examplepluginconfig)<br>Configuration for the exampleplugin plugin
-`secrets` | [SecretsConfig](#secretsconfig)<br>Configuration for the secrets plugin
 `experimental` | [ExperimentalConfig](#experimentalconfig)<br>Configuration for the experimental plugin
+`secrets` | [SecretsConfig](#secretsconfig)<br>Configuration for the secrets plugin
 `symbols` | [SymbolsConfig](#symbolsconfig)<br>Configuration for the symbols plugin
-`storage` | [StorageConfig](#storageconfig)<br>Configuration for the storage plugin
 `tools` | [ToolsConfig](#toolsconfig)<br>Configuration for the tools plugin
+`build` | [BuildConfig](#buildconfig)<br>Configuration for the build plugin
+`compute` | [ComputeConfig](#computeconfig)<br>Configuration for the compute plugin
+`storage` | [StorageConfig](#storageconfig)<br>Configuration for the storage plugin
 
 ## AnalyticsConfig
 
@@ -77,23 +77,63 @@
 ---- | -----------
 `stores` | [TelemetryStoreConfig](Telemetry.md)`[]`<br>Metrics to aggregate on the Horde server
 
-## ComputeConfig
+## ExamplePluginConfig
 
-Configuration for the compute system
+Configuration for the ExamplePlugin plugin
 
 Name | Description
 ---- | -----------
-`include` | [ConfigInclude](#configinclude)`[]`<br>Includes for other configuration files
-`macros` | [ConfigMacro](#configmacro)`[]`<br>Macros within the compute config
-`acl` | [AclConfig](#aclconfig)<br>Inherited root acl
-`versionEnum` | [ConfigVersion](#configversion-enum)<br>Config version number
-`telemetryStoreId` | `string`<br>Telemetry store for lease and compute telemetry data
-`rates` | [AgentRateConfig](#agentrateconfig)`[]`<br>List of costs of a particular agent type
-`clusters` | [ComputeClusterConfig](#computeclusterconfig)`[]`<br>List of compute profiles
-`pools` | [PoolConfig](#poolconfig)`[]`<br>List of pools
-`software` | [AgentSoftwareConfig](#agentsoftwareconfig)`[]`<br>List of costs of a particular agent type
-`networks` | [NetworkConfig](#networkconfig)`[]`<br>List of networks
+`greetingMessage` | `string`<br>The Greeting provided by the example plugin
 
+## ExperimentalConfig
+
+Configuration for the experimental plugin
+
+Name | Description
+---- | -----------
+`notifications` | [NotificationConfig](Notification.md)`[]`<br>Notifications to be sent on the Horde server
+
+## SecretsConfig
+
+Configuration for the secrets system
+
+Name | Description
+---- | -----------
+`secrets` | [SecretConfig](#secretconfig)`[]`<br>List of secrets
+`providerConfigs` | [SecretProviderConfig](#secretproviderconfig)`[]`<br>List of configurations for secret providers
+
+## SecretConfig
+
+Configuration for a secret value
+
+Name | Description
+---- | -----------
+`id` | `string`<br>Identifier for this secret
+`data` | `string` `->` `string`<br>Key/value pairs associated with this secret
+`sources` | [ExternalSecretConfig](#externalsecretconfig)`[]`<br>Providers to source key/value pairs from
+`acl` | [AclConfig](#aclconfig)<br>Defines access to this particular secret
+
+## ExternalSecretConfig
+
+Configuration for an external secret provider
+
+Name | Description
+---- | -----------
+`providerConfig` | `string`<br>The name to look up a secret provider config. A config is mutually exclusive to and to be used when a secret requires additional options than just the name of the provider
+`provider` | `string`<br>Name of the provider to use
+`format` | [ExternalSecretFormat](#externalsecretformat-enum)<br>Format of the secret
+`key` | `string`<br>Optional key indicating the parameter to set in the resulting data array. Required if  is .
+`path` | `string`<br>Optional value indicating what to fetch from the provider
+
+## ExternalSecretFormat (Enum)
+
+Format describing how to parse external secret values
+
+Name | Description
+---- | -----------
+`Text` | Secret is a plain text value which will be stored using the external secret key
+`Json` | Secret is a JSON formatted string containing key/value pairs
+
 ## AclConfig
 
 Parameters to update an ACL
@@ -135,187 +175,94 @@
 `excludeActions` | `string[]`<br>Actions to exclude from the inherited actions
 `extends` | `string[]`<br>Other profiles to extend from
 
-## ConfigVersion (Enum)
+## SecretProviderConfig
 
-Global version number for running the server. As new features are introduced that require data migrations, this version number indicates the backwards compatibility functionality that must be enabled. When adding a new version here, also add a message to ConfigService.CreateSnapshotAsync describing the steps that need to be taken to upgrade the deployment.
+Common settings object for different secret providers
 
 Name | Description
 ---- | -----------
-`None` | Not specified
-`Initial` | Initial version number
-`PoolsInConfigFiles` | Ability to add/remove pools via the REST API is removed. Pools should be configured through globals.json instead.
-`IncrementalPoolMigration` | Safe hybrid mode for incremental pool migration. Pools are read from both config files and the database, with deduplication. Config file pools take priority — database pools with matching IDs are skipped. The pool editor UI remains visible for database-managed pools. Use /api/v1/server/migrate/pool-config to export remaining database pools.
-`Latest` | Latest version number
-`LatestPlusOne` | One after the last defined version number
+`name` | `string`<br>Name of the secret provider config
+`provider` | `string`<br>Name of the provider
+`hcpVault` | [HcpVaultConfig](#hcpvaultconfig)<br>Configuration for HCP Vault
+`azureKeyVault` | [AzureKeyVaultConfig](#azurekeyvaultconfig)<br>Configuration for Azure Key Vault
 
-## AgentRateConfig
+## HcpVaultConfig
 
-Describes the monetary cost of agents matching a particular criteria
+Configuration for HCP Vault secrets.
 
 Name | Description
 ---- | -----------
-`condition` | `string`<br>Condition string
-`rate` | `number`<br>Rate for this agent
+`preSharedKey` | `string`<br>The Vault authentication token to use when the type of credentials is .<br>More information on this token can be found in the API documentation for X-Vault-Token
+`endPoint` | `string`<br>The address of the Vault server.
+`awsArnRole` | `string`<br>The optional AWS ARN to assume a role when logging into Vault. To be used when the type of credentials is . When not set the AWS SDK default credential search will be used.
+`awsIamServerId` | `string`<br>The value for the X-Vault-AWS-IAM-Server-ID header when the Vault server is configured with iam_server_id_header_value. Most likely to be the Vault server's DNS name. To be used when the type of credentials is .
+`role` | `string`<br>The name of the role in the Vault configuration for the login. This role name can be different to the name of the AWS IAM principal. To be used when the type of credentials is . When not set the name of IAM principal will be used.
+`credentials` | [HcpVaultCredentialsType](#hcpvaultcredentialstype-enum)<br>Type of credentials to use.
 
-## ComputeClusterConfig
+## HcpVaultCredentialsType (Enum)
 
-Profile for executing compute requests
+Credentials to use to log into HCP Vault.
 
 Name | Description
 ---- | -----------
-`id` | `string`<br>Name of the partition
-`namespaceId` | `string`<br>Name of the namespace to use
-`requestBucketId` | `string`<br>Name of the input bucket
-`responseBucketId` | `string`<br>Name of the output bucket
-`condition` | `string`<br>Filter for agents to include
-`acl` | [AclConfig](#aclconfig)<br>Access control list
-`uba` | [UbaComputeClusterConfig](#ubacomputeclusterconfig)<br>Configuration settings for Unreal Build Accelerator (UBA) for a compute cluster
+`PreSharedKey` | Use a pre-shared token provided by Vault for authentication.
+`AwsAuth` | Use AWS IAM credentials to authenticate against the Vault server.
 
-## UbaComputeClusterConfig
+## AzureKeyVaultConfig
 
-Configuration settings for Unreal Build Accelerator (UBA) for a compute cluster Contains cache server connectivity and other UBA-related settings. Ensure all cache related fields are non-null to activate use of cache.
+Configuration for Azure Key Vault secrets.
 
 Name | Description
 ---- | -----------
-`cacheEndpoint` | `string`<br>Host and port for connecting to the UBA cache server (host:port format) As specified with `-port=[PORT]` to UBA cache process. Serves all cache requests and must be accessible to both agents and initiators running UBA / Unreal Build Tool.
-`cacheHttpEndpoint` | `string`<br>Host and port for HTTP management of UBA cache server (host:port format) As specified with `-http=[PORT]` to UBA cache process. This port is used by Horde to manage the UBA cache, such as adding new cache session keys for initiators For improved security, only allow Horde server to access this port
-`cacheHttpCrypto` | `string`<br>AES encryption key used for secure Horde-to-cache-server communication As specified with `-httpcrypto=[KEY]` to UBA cache process.
+`vaultUri` | `string`<br>The Uri to the Azure Key Vault.
 
-## PoolConfig
+## SymbolsConfig
 
-Mutable configuration for a pool
+Configuration for the tools system
 
 Name | Description
 ---- | -----------
-`id` | `string`<br>Unique id for this pool
-`base` | `string`<br>Base pool config to copy settings from
-`name` | `string`<br>Name of the pool
-`condition` | `string`<br>Condition for agents to automatically be included in this pool
-`properties` | `string` `->` `string`<br>Arbitrary properties related to this pool
-`color` | [PoolColor](#poolcolor-enum)<br>Color to use for this pool on the dashboard
-`enableAutoscaling` | `boolean`<br>Whether to enable autoscaling for this pool
-`minAgents` | `integer`<br>The minimum number of agents to keep in the pool
-`numReserveAgents` | `integer`<br>The minimum number of idle agents to hold in reserve
-`conformInterval` | `string`<br>Interval between conforms. If zero, the pool will not conform on a schedule.
-`scaleOutCooldown` | `string`<br>Cooldown time between scale-out events
-`scaleInCooldown` | `string`<br>Cooldown time between scale-in events
-`shutdownIfDisabledGracePeriod` | `string`<br>Time to wait before shutting down an agent that has been disabled
-`sizeStrategy` | [PoolSizeStrategy](#poolsizestrategy-enum)<br>
-`sizeStrategies` | [PoolSizeStrategyInfo](#poolsizestrategyinfo)`[]`<br>List of pool sizing strategies for this pool. The first strategy with a matching condition will be picked.
-`fleetManagers` | [FleetManagerInfo](#fleetmanagerinfo)`[]`<br>List of fleet managers for this pool. The first strategy with a matching condition will be picked. If empty or no conditions match, a default fleet manager will be used.
-`leaseUtilizationSettings` | [LeaseUtilizationSettings](#leaseutilizationsettings)<br>Settings for lease utilization pool sizing strategy (if used)
-`jobQueueSettings` | [JobQueueSettings](#jobqueuesettings)<br>Settings for job queue pool sizing strategy (if used)
-`computeQueueAwsMetricSettings` | [ComputeQueueAwsMetricSettings](#computequeueawsmetricsettings)<br>
+`stores` | [SymbolStoreConfig](#symbolstoreconfig)`[]`<br>List of symbol stores
 
-## PoolColor (Enum)
+## SymbolStoreConfig
 
-Color to use for labels of this pool
+Configuration for a symbol store
 
 Name | Description
 ---- | -----------
-`Default` | 
-`Blue` | 
-`Orange` | 
-`Green` | 
-`Gray` | 
+`id` | `string`<br>Identifier for this store
+`namespaceId` | `string`<br>Configuration for the symbol store backend
+`public` | `boolean`<br>Whether to make this store available without auth
+`acl` | [AclConfig](#aclconfig)<br>Access to the symbol store
 
-## PoolSizeStrategy (Enum)
+## ToolsConfig
 
-Available pool sizing strategies
+Configuration for the tools system
 
 Name | Description
 ---- | -----------
-`LeaseUtilization` | Strategy based on lease utilization
-`JobQueue` | Strategy based on size of job build queue
-`NoOp` | No-op strategy used as fallback/default behavior
-`ComputeQueueAwsMetric` | A no-op strategy that reports metrics to let an external AWS auto-scaling policy scale the fleet
-`LeaseUtilizationAwsMetric` | A no-op strategy that reports metrics to let an external AWS auto-scaling policy scale the fleet
+`tools` | [ToolConfig](#toolconfig)`[]`<br>Tool configurations
 
-## PoolSizeStrategyInfo
+## ToolConfig
 
-Metadata for configuring and picking a pool sizing strategy
+Options for configuring a tool
 
 Name | Description
 ---- | -----------
-`type` | [PoolSizeStrategy](#poolsizestrategy-enum)<br>Strategy implementation to use
-`condition` | `string`<br>Condition if this strategy should be enabled (right now, using date/time as a distinguishing factor)
-`config` | `object`<br>Configuration for the strategy, serialized as JSON
-`extraAgentCount` | `integer`<br>Integer to add after pool size has been calculated. Can also be negative.
+`id` | `string`<br>Unique identifier for the tool
+`name` | `string`<br>Name of the tool
+`description` | `string`<br>Description for the tool
+`category` | `string`<br>Category for the tool. Will cause the tool to be shown in a different tab in the dashboard.
+`group` | `string`<br>Grouping key for different variations of the same tool. The dashboard will show these together.
+`platforms` | `string[]`<br>Platforms for this tool. Takes the form of a NET RID (https://learn.microsoft.com/en-us/dotnet/core/rid-catalog).
+`public` | `boolean`<br>Whether this tool should be exposed for download on a public endpoint without authentication
+`showInUgs` | `boolean`<br>Whether to show this tool for download in the UGS tools menu
+`showInDashboard` | `boolean`<br>Whether to show this tool for download in the dashboard
+`showInToolbox` | `boolean`<br>Whether to show this tool for download in Unreal Toolbox
+`metadata` | `string` `->` `string`<br>Metadata for this tool
+`namespaceId` | `string`<br>Default namespace for new deployments of this tool
+`maxDeployments` | `integer`<br>Maximum number of deployments to retain for this tool. When exceeded, the oldest deployments are removed.
+`acl` | [AclConfig](#aclconfig)<br>Permissions for the tool
 
-## FleetManagerInfo
-
-Metadata for configuring and picking a fleet manager
-
-Name | Description
----- | -----------
-`type` | [FleetManagerType](#fleetmanagertype-enum)<br>Fleet manager type implementation to use
-`condition` | `string`<br>Condition if this strategy should be enabled (right now, using date/time as a distinguishing factor)
-`config` | `object`<br>Configuration for the strategy, serialized as JSON
-
-## FleetManagerType (Enum)
-
-Available fleet managers
-
-Name | Description
----- | -----------
-`Default` | Default fleet manager
-`NoOp` | No-op fleet manager.
-`Aws` | Fleet manager for handling AWS EC2 instances. Will create and/or terminate instances from scratch.
-`AwsReuse` | Fleet manager for handling AWS EC2 instances. Will start already existing but stopped instances to reuse existing EBS disks.
-`AwsRecycle` | Fleet manager for handling AWS EC2 instances. Will start already existing but stopped instances to reuse existing EBS disks.
-`AwsAsg` | Fleet manager for handling AWS EC2 instances. Uses an EC2 auto-scaling group for controlling the number of running instances.
-
-## LeaseUtilizationSettings
-
-Lease utilization sizing settings for a pool
-
-Name | Description
----- | -----------
-`sampleTimeSec` | `integer`<br>Time period for each sample
-`numSamples` | `integer`<br>Number of samples to collect for calculating lease utilization
-`numSamplesForResult` | `integer`<br>Min number of samples for a valid result
-`minAgents` | `integer`<br>The minimum number of agents to keep in the pool
-`numReserveAgents` | `integer`<br>The minimum number of idle agents to hold in reserve
-
-## JobQueueSettings
-
-Job queue sizing settings for a pool
-
-Name | Description
----- | -----------
-`scaleOutFactor` | `number`<br>Factor translating queue size to additional agents to grow the pool with The result is always rounded up to nearest integer. Example: if there are 20 jobs in queue, a factor 0.25 will result in 5 new agents being added (20 * 0.25)
-`scaleInFactor` | `number`<br>Factor by which to shrink the pool size with when queue is empty The result is always rounded up to nearest integer. Example: when the queue size is zero, a default value of 0.9 will shrink the pool by 10% (current agent count * 0.9)
-`samplePeriodMin` | `integer`<br>How far back in time to look for job batches (that potentially are in the queue)
-`readyTimeThresholdSec` | `integer`<br>Time spent in ready state before considered truly waiting for an agent<br>A job batch can be in ready state before getting picked up and executed. This threshold will help ensure only batches that have been waiting longer than this value will be considered.
-
-## ComputeQueueAwsMetricSettings
-
-Settings for
-
-Name | Description
----- | -----------
-`computeClusterId` | `string`<br>Compute cluster ID to observe
-`namespace` | `string`<br>AWS CloudWatch namespace to write metrics in
-
-## AgentSoftwareConfig
-
-Selects different agent software versions by evaluating a condition
-
-Name | Description
----- | -----------
-`toolId` | `string`<br>Tool identifier
-`condition` | `string`<br>Condition for using this channel
-
-## NetworkConfig
-
-Describes a network The ID describes any logical grouping, such as region, availability zone, rack or office location.
-
-Name | Description
----- | -----------
-`id` | `string`<br>ID for this network
-`cidrBlock` | `string`<br>CIDR block
-`description` | `string`<br>Human-readable description
-`computeId` | `string`<br>Compute ID for this network (used when allocating compute resources)
-
 ## BuildConfig
 
@@ -463,122 +410,204 @@
 `keepDays` | `integer`<br>Number of days to retain artifacts of this type
 `namespaceId` | `string`<br>Storage namespace to use for this artifact types
 
-## ExamplePluginConfig
+## ComputeConfig
 
-Configuration for the ExamplePlugin plugin
+Configuration for the compute system
 
 Name | Description
 ---- | -----------
-`greetingMessage` | `string`<br>The Greeting provided by the example plugin
+`include` | [ConfigInclude](#configinclude)`[]`<br>Includes for other configuration files
+`macros` | [ConfigMacro](#configmacro)`[]`<br>Macros within the compute config
+`acl` | [AclConfig](#aclconfig)<br>Inherited root acl
+`versionEnum` | [ConfigVersion](#configversion-enum)<br>Config version number
+`telemetryStoreId` | `string`<br>Telemetry store for lease and compute telemetry data
+`rates` | [AgentRateConfig](#agentrateconfig)`[]`<br>List of costs of a particular agent type
+`clusters` | [ComputeClusterConfig](#computeclusterconfig)`[]`<br>List of compute profiles
+`pools` | [PoolConfig](#poolconfig)`[]`<br>List of pools
+`software` | [AgentSoftwareConfig](#agentsoftwareconfig)`[]`<br>List of costs of a particular agent type
+`networks` | [NetworkConfig](#networkconfig)`[]`<br>List of networks
 
-## SecretsConfig
+## ConfigVersion (Enum)
 
-Configuration for the secrets system
+Global version number for running the server. As new features are introduced that require data migrations, this version number indicates the backwards compatibility functionality that must be enabled. When adding a new version here, also add a message to ConfigService.CreateSnapshotAsync describing the steps that need to be taken to upgrade the deployment.
 
 Name | Description
 ---- | -----------
-`secrets` | [SecretConfig](#secretconfig)`[]`<br>List of secrets
-`providerConfigs` | [SecretProviderConfig](#secretproviderconfig)`[]`<br>List of configurations for secret providers
+`None` | Not specified
+`Initial` | Initial version number
+`PoolsInConfigFiles` | Ability to add/remove pools via the REST API is removed. Pools should be configured through globals.json instead.
+`IncrementalPoolMigration` | Safe hybrid mode for incremental pool migration. Pools are read from both config files and the database, with deduplication. Config file pools take priority — database pools with matching IDs are skipped. The pool editor UI remains visible for database-managed pools. Use /api/v1/server/migrate/pool-config to export remaining database pools.
+`Latest` | Latest version number
+`LatestPlusOne` | One after the last defined version number
 
-## SecretConfig
+## AgentRateConfig
 
-Configuration for a secret value
+Describes the monetary cost of agents matching a particular criteria
 
 Name | Description
 ---- | -----------
-`id` | `string`<br>Identifier for this secret
-`data` | `string` `->` `string`<br>Key/value pairs associated with this secret
-`sources` | [ExternalSecretConfig](#externalsecretconfig)`[]`<br>Providers to source key/value pairs from
-`acl` | [AclConfig](#aclconfig)<br>Defines access to this particular secret
+`condition` | `string`<br>Condition string
+`rate` | `number`<br>Rate for this agent
 
-## ExternalSecretConfig
+## ComputeClusterConfig
 
-Configuration for an external secret provider
+Profile for executing compute requests
 
 Name | Description
 ---- | -----------
-`providerConfig` | `string`<br>The name to look up a secret provider config. A config is mutually exclusive to and to be used when a secret requires additional options than just the name of the provider
-`provider` | `string`<br>Name of the provider to use
-`format` | [ExternalSecretFormat](#externalsecretformat-enum)<br>Format of the secret
-`key` | `string`<br>Optional key indicating the parameter to set in the resulting data array. Required if  is .
-`path` | `string`<br>Optional value indicating what to fetch from the provider
+`id` | `string`<br>Name of the partition
+`namespaceId` | `string`<br>Name of the namespace to use
+`requestBucketId` | `string`<br>Name of the input bucket
+`responseBucketId` | `string`<br>Name of the output bucket
+`condition` | `string`<br>Filter for agents to include
+`acl` | [AclConfig](#aclconfig)<br>Access control list
+`uba` | [UbaComputeClusterConfig](#ubacomputeclusterconfig)<br>Configuration settings for Unreal Build Accelerator (UBA) for a compute cluster
 
-## ExternalSecretFormat (Enum)
+## UbaComputeClusterConfig
 
-Format describing how to parse external secret values
+Configuration settings for Unreal Build Accelerator (UBA) for a compute cluster Contains cache server connectivity and other UBA-related settings. Ensure all cache related fields are non-null to activate use of cache.
 
 Name | Description
 ---- | -----------
-`Text` | Secret is a plain text value which will be stored using the external secret key
-`Json` | Secret is a JSON formatted string containing key/value pairs
+`cacheEndpoint` | `string`<br>Host and port for connecting to the UBA cache server (host:port format) As specified with `-port=[PORT]` to UBA cache process. Serves all cache requests and must be accessible to both agents and initiators running UBA / Unreal Build Tool.
+`cacheHttpEndpoint` | `string`<br>Host and port for HTTP management of UBA cache server (host:port format) As specified with `-http=[PORT]` to UBA cache process. This port is used by Horde to manage the UBA cache, such as adding new cache session keys for initiators For improved security, only allow Horde server to access this port
+`cacheHttpCrypto` | `string`<br>AES encryption key used for secure Horde-to-cache-server communication As specified with `-httpcrypto=[KEY]` to UBA cache process.
 
-## SecretProviderConfig
+## PoolConfig
 
-Common settings object for different secret providers
+Mutable configuration for a pool
 
 Name | Description
 ---- | -----------
-`name` | `string`<br>Name of the secret provider config
-`provider` | `string`<br>Name of the provider
-`hcpVault` | [HcpVaultConfig](#hcpvaultconfig)<br>Configuration for HCP Vault
-`azureKeyVault` | [AzureKeyVaultConfig](#azurekeyvaultconfig)<br>Configuration for Azure Key Vault
+`id` | `string`<br>Unique id for this pool
+`base` | `string`<br>Base pool config to copy settings from
+`name` | `string`<br>Name of the pool
+`condition` | `string`<br>Condition for agents to automatically be included in this pool
+`properties` | `string` `->` `string`<br>Arbitrary properties related to this pool
+`color` | [PoolColor](#poolcolor-enum)<br>Color to use for this pool on the dashboard
+`enableAutoscaling` | `boolean`<br>Whether to enable autoscaling for this pool
+`minAgents` | `integer`<br>The minimum number of agents to keep in the pool
+`numReserveAgents` | `integer`<br>The minimum number of idle agents to hold in reserve
+`conformInterval` | `string`<br>Interval between conforms. If zero, the pool will not conform on a schedule.
+`scaleOutCooldown` | `string`<br>Cooldown time between scale-out events
+`scaleInCooldown` | `string`<br>Cooldown time between scale-in events
+`shutdownIfDisabledGracePeriod` | `string`<br>Time to wait before shutting down an agent that has been disabled
+`sizeStrategy` | [PoolSizeStrategy](#poolsizestrategy-enum)<br>
+`sizeStrategies` | [PoolSizeStrategyInfo](#poolsizestrategyinfo)`[]`<br>List of pool sizing strategies for this pool. The first strategy with a matching condition will be picked.
+`fleetManagers` | [FleetManagerInfo](#fleetmanagerinfo)`[]`<br>List of fleet managers for this pool. The first strategy with a matching condition will be picked. If empty or no conditions match, a default fleet manager will be used.
+`leaseUtilizationSettings` | [LeaseUtilizationSettings](#leaseutilizationsettings)<br>Settings for lease utilization pool sizing strategy (if used)
+`jobQueueSettings` | [JobQueueSettings](#jobqueuesettings)<br>Settings for job queue pool sizing strategy (if used)
+`computeQueueAwsMetricSettings` | [ComputeQueueAwsMetricSettings](#computequeueawsmetricsettings)<br>
 
-## HcpVaultConfig
+## PoolColor (Enum)
 
-Configuration for HCP Vault secrets.
+Color to use for labels of this pool
 
 Name | Description
 ---- | -----------
-`preSharedKey` | `string`<br>The Vault authentication token to use when the type of credentials is .<br>More information on this token can be found in the API documentation for X-Vault-Token
-`endPoint` | `string`<br>The address of the Vault server.
-`awsArnRole` | `string`<br>The optional AWS ARN to assume a role when logging into Vault. To be used when the type of credentials is . When not set the AWS SDK default credential search will be used.
-`awsIamServerId` | `string`<br>The value for the X-Vault-AWS-IAM-Server-ID header when the Vault server is configured with iam_server_id_header_value. Most likely to be the Vault server's DNS name. To be used when the type of credentials is .
-`role` | `string`<br>The name of the role in the Vault configuration for the login. This role name can be different to the name of the AWS IAM principal. To be used when the type of credentials is . When not set the name of IAM principal will be used.
-`credentials` | [HcpVaultCredentialsType](#hcpvaultcredentialstype-enum)<br>Type of credentials to use.
+`Default` | 
+`Blue` | 
+`Orange` | 
+`Green` | 
+`Gray` | 
 
-## HcpVaultCredentialsType (Enum)
+## PoolSizeStrategy (Enum)
 
-Credentials to use to log into HCP Vault.
+Available pool sizing strategies
 
 Name | Description
 ---- | -----------
-`PreSharedKey` | Use a pre-shared token provided by Vault for authentication.
-`AwsAuth` | Use AWS IAM credentials to authenticate against the Vault server.
+`LeaseUtilization` | Strategy based on lease utilization
+`JobQueue` | Strategy based on size of job build queue
+`NoOp` | No-op strategy used as fallback/default behavior
+`ComputeQueueAwsMetric` | A no-op strategy that reports metrics to let an external AWS auto-scaling policy scale the fleet
+`LeaseUtilizationAwsMetric` | A no-op strategy that reports metrics to let an external AWS auto-scaling policy scale the fleet
 
-## AzureKeyVaultConfig
+## PoolSizeStrategyInfo
 
-Configuration for Azure Key Vault secrets.
+Metadata for configuring and picking a pool sizing strategy
 
 Name | Description
 ---- | -----------
-`vaultUri` | `string`<br>The Uri to the Azure Key Vault.
+`type` | [PoolSizeStrategy](#poolsizestrategy-enum)<br>Strategy implementation to use
+`condition` | `string`<br>Condition if this strategy should be enabled (right now, using date/time as a distinguishing factor)
+`config` | `object`<br>Configuration for the strategy, serialized as JSON
+`extraAgentCount` | `integer`<br>Integer to add after pool size has been calculated. Can also be negative.
 
-## ExperimentalConfig
+## FleetManagerInfo
 
-Configuration for the experimental plugin
+Metadata for configuring and picking a fleet manager
 
 Name | Description
 ---- | -----------
-`notifications` | [NotificationConfig](Notification.md)`[]`<br>Notifications to be sent on the Horde server
+`type` | [FleetManagerType](#fleetmanagertype-enum)<br>Fleet manager type implementation to use
+`condition` | `string`<br>Condition if this strategy should be enabled (right now, using date/time as a distinguishing factor)
+`config` | `object`<br>Configuration for the strategy, serialized as JSON
 
-## SymbolsConfig
+## FleetManagerType (Enum)
 
-Configuration for the tools system
+Available fleet managers
 
 Name | Description
 ---- | -----------
-`stores` | [SymbolStoreConfig](#symbolstoreconfig)`[]`<br>List of symbol stores
+`Default` | Default fleet manager
+`NoOp` | No-op fleet manager.
+`Aws` | Fleet manager for handling AWS EC2 instances. Will create and/or terminate instances from scratch.
+`AwsReuse` | Fleet manager for handling AWS EC2 instances. Will start already existing but stopped instances to reuse existing EBS disks.
+`AwsRecycle` | Fleet manager for handling AWS EC2 instances. Will start already existing but stopped instances to reuse existing EBS disks.
+`AwsAsg` | Fleet manager for handling AWS EC2 instances. Uses an EC2 auto-scaling group for controlling the number of running instances.
 
-## SymbolStoreConfig
+## LeaseUtilizationSettings
 
-Configuration for a symbol store
+Lease utilization sizing settings for a pool
 
 Name | Description
 ---- | -----------
-`id` | `string`<br>Identifier for this store
-`namespaceId` | `string`<br>Configuration for the symbol store backend
-`public` | `boolean`<br>Whether to make this store available without auth
-`acl` | [AclConfig](#aclconfig)<br>Access to the symbol store
+`sampleTimeSec` | `integer`<br>Time period for each sample
+`numSamples` | `integer`<br>Number of samples to collect for calculating lease utilization
+`numSamplesForResult` | `integer`<br>Min number of samples for a valid result
+`minAgents` | `integer`<br>The minimum number of agents to keep in the pool
+`numReserveAgents` | `integer`<br>The minimum number of idle agents to hold in reserve
 
+## JobQueueSettings
+
+Job queue sizing settings for a pool
+
+Name | Description
+---- | -----------
+`scaleOutFactor` | `number`<br>Factor translating queue size to additional agents to grow the pool with The result is always rounded up to nearest integer. Example: if there are 20 jobs in queue, a factor 0.25 will result in 5 new agents being added (20 * 0.25)
+`scaleInFactor` | `number`<br>Factor by which to shrink the pool size with when queue is empty The result is always rounded up to nearest integer. Example: when the queue size is zero, a default value of 0.9 will shrink the pool by 10% (current agent count * 0.9)
+`samplePeriodMin` | `integer`<br>How far back in time to look for job batches (that potentially are in the queue)
+`readyTimeThresholdSec` | `integer`<br>Time spent in ready state before considered truly waiting for an agent<br>A job batch can be in ready state before getting picked up and executed. This threshold will help ensure only batches that have been waiting longer than this value will be considered.
+
+## ComputeQueueAwsMetricSettings
+
+Settings for
+
+Name | Description
+---- | -----------
+`computeClusterId` | `string`<br>Compute cluster ID to observe
+`namespace` | `string`<br>AWS CloudWatch namespace to write metrics in
+
+## AgentSoftwareConfig
+
+Selects different agent software versions by evaluating a condition
+
+Name | Description
+---- | -----------
+`toolId` | `string`<br>Tool identifier
+`condition` | `string`<br>Condition for using this channel
+
+## NetworkConfig
+
+Describes a network The ID describes any logical grouping, such as region, availability zone, rack or office location.
+
+Name | Description
+---- | -----------
+`id` | `string`<br>ID for this network
+`cidrBlock` | `string`<br>CIDR block
+`description` | `string`<br>Human-readable description
+`computeId` | `string`<br>Compute ID for this network (used when allocating compute resources)
+
 ## StorageConfig
 
@@ -657,32 +686,3 @@
 `gcDelayHrs` | `number`<br>How long to keep newly uploaded orphanned objects before allowing them to be deleted, in hours.
 `enableAliases` | `boolean`<br>Support querying exports by their aliases
 `acl` | [AclConfig](#aclconfig)<br>Access list for this namespace
-
-## ToolsConfig
-
-Configuration for the tools system
-
-Name | Description
----- | -----------
-`tools` | [ToolConfig](#toolconfig)`[]`<br>Tool configurations
-
-## ToolConfig
-
-Options for configuring a tool
-
-Name | Description
----- | -----------
-`id` | `string`<br>Unique identifier for the tool
-`name` | `string`<br>Name of the tool
-`description` | `string`<br>Description for the tool
-`category` | `string`<br>Category for the tool. Will cause the tool to be shown in a different tab in the dashboard.
-`group` | `string`<br>Grouping key for different variations of the same tool. The dashboard will show these together.
-`platforms` | `string[]`<br>Platforms for this tool. Takes the form of a NET RID (https://learn.microsoft.com/en-us/dotnet/core/rid-catalog).
-`public` | `boolean`<br>Whether this tool should be exposed for download on a public endpoint without authentication
-`showInUgs` | `boolean`<br>Whether to show this tool for download in the UGS tools menu
-`showInDashboard` | `boolean`<br>Whether to show this tool for download in the dashboard
-`showInToolbox` | `boolean`<br>Whether to show this tool for download in Unreal Toolbox
-`metadata` | `string` `->` `string`<br>Metadata for this tool
-`namespaceId` | `string`<br>Default namespace for new deployments of this tool
-`maxDeployments` | `integer`<br>Maximum number of deployments to retain for this tool. When exceeded, the oldest deployments are removed.
-`acl` | [AclConfig](#aclconfig)<br>Permissions for the tool
diff --git a/Engine/Source/Programs/Horde/Docs/Deployment/ServerSettings.md b/Engine/Source/Programs/Horde/Docs/Deployment/ServerSettings.md
--- a/Engine/Source/Programs/Horde/Docs/Deployment/ServerSettings.md
+++ b/Engine/Source/Programs/Horde/Docs/Deployment/ServerSettings.md
@@ -128,14 +128,14 @@
 Name | Description
 ---- | -----------
 `analytics` | [AnalyticsServerConfig](#analyticsserverconfig)<br>Configuration for the analytics plugin
-`compute` | [ComputeServerConfig](#computeserverconfig)<br>Configuration for the compute plugin
-`build` | [BuildServerConfig](#buildserverconfig)<br>Configuration for the build plugin
 `exampleplugin` | [ExamplePluginServerConfig](#examplepluginserverconfig)<br>Configuration for the exampleplugin plugin
-`secrets` | [SecretsServerConfig](#secretsserverconfig)<br>Configuration for the secrets plugin
 `experimental` | [ExperimentalServerConfig](#experimentalserverconfig)<br>Configuration for the experimental plugin
+`secrets` | [SecretsServerConfig](#secretsserverconfig)<br>Configuration for the secrets plugin
 `symbols` | [PluginServerConfig](#pluginserverconfig)<br>Configuration for the symbols plugin
-`storage` | [StorageServerConfig](#storageserverconfig)<br>Configuration for the storage plugin
 `tools` | [ToolsServerConfig](#toolsserverconfig)<br>Configuration for the tools plugin
+`build` | [BuildServerConfig](#buildserverconfig)<br>Configuration for the build plugin
+`compute` | [ComputeServerConfig](#computeserverconfig)<br>Configuration for the compute plugin
+`storage` | [StorageServerConfig](#storageserverconfig)<br>Configuration for the storage plugin
 
 ## AnalyticsServerConfig
 
@@ -174,50 +174,112 @@
 `retainDays` | `number`<br>Number of days worth of telmetry events to keep
 `enabled` | `boolean`<br>Whether to enable this sink
 
-## ComputeServerConfig
+## ExamplePluginServerConfig
 
-Static configuration for the compute plugin
+Server configuration for the ExamplePlugin
 
 Name | Description
 ---- | -----------
-`enableUpgradeTasks` | `boolean`<br>Whether to enable the upgrade task source, always upgrading agents to the latest version
+`enabled` | `boolean`<br>Whether the plugin should be enabled or not
+
+## ExperimentalServerConfig
+
+Server configuration for the experimental plugin
+
+Name | Description
+---- | -----------
+`enabled` | `boolean`<br>Whether the plugin should be enabled or not
+
+## SecretsServerConfig
+
+Static configuration for the secrets plugin
+
+Name | Description
+---- | -----------
 `withAws` | `boolean`<br>Whether to enable Amazon Web Services (AWS) specific features
-`awsRegions` | `string[]`<br>List of AWS regions for Horde to be aware of (e.g. us-east-1 or eu-central-1) Right now, this is only used for replicating CloudWatch metrics to multiple regions
-`awsAutoScalingQueueUrls` | `string[]`<br>AWS SQS queue URLs where lifecycle events from EC2 auto-scaling are received
-`fleetManagerV2` | [FleetManagerType](#fleetmanagertype-enum)<br>Default fleet manager to use (when not specified by pool)
-`fleetManagerV2Config` | `object`<br>Config for the fleet manager (serialized JSON)
-`autoEnrollAgents` | `boolean`<br>Whether to automatically enroll agents in the farm
-`defaultAgentPoolSizeStrategy` | [PoolSizeStrategy](#poolsizestrategy-enum)<br>Default agent pool sizing strategy for pools that doesn't have one explicitly configured
-`agentPoolScaleOutCooldownSeconds` | `integer`<br>Scale-out cooldown for auto-scaling agent pools (in seconds). Can be overridden by per-pool settings.
-`agentPoolScaleInCooldownSeconds` | `integer`<br>Scale-in cooldown for auto-scaling agent pools (in seconds). Can be overridden by per-pool settings.
-`computeTunnelPort` | `integer`<br>Port to listen on for tunneling compute sockets to agents
-`computeTunnelAddress` | `string`<br>What address (host:port) clients should connect to for compute socket tunneling Port may differ from  if Horde server is behind a reverse proxy/firewall
 `enabled` | `boolean`<br>Whether the plugin should be enabled or not
 
-## FleetManagerType (Enum)
+## PluginServerConfig
 
-Available fleet managers
+Base class for plugin server config objects
 
 Name | Description
 ---- | -----------
-`Default` | Default fleet manager
-`NoOp` | No-op fleet manager.
-`Aws` | Fleet manager for handling AWS EC2 instances. Will create and/or terminate instances from scratch.
-`AwsReuse` | Fleet manager for handling AWS EC2 instances. Will start already existing but stopped instances to reuse existing EBS disks.
-`AwsRecycle` | Fleet manager for handling AWS EC2 instances. Will start already existing but stopped instances to reuse existing EBS disks.
-`AwsAsg` | Fleet manager for handling AWS EC2 instances. Uses an EC2 auto-scaling group for controlling the number of running instances.
+`enabled` | `boolean`<br>Whether the plugin should be enabled or not
 
-## PoolSizeStrategy (Enum)
+## ToolsServerConfig
 
-Available pool sizing strategies
+Server configuration for bundled tools
 
 Name | Description
 ---- | -----------
-`LeaseUtilization` | Strategy based on lease utilization
-`JobQueue` | Strategy based on size of job build queue
-`NoOp` | No-op strategy used as fallback/default behavior
-`ComputeQueueAwsMetric` | A no-op strategy that reports metrics to let an external AWS auto-scaling policy scale the fleet
-`LeaseUtilizationAwsMetric` | A no-op strategy that reports metrics to let an external AWS auto-scaling policy scale the fleet
+`bundledTools` | [BundledToolConfig](#bundledtoolconfig)`[]`<br>Tools bundled along with the server. Data for each tool can be produced using the 'bundle create' command, and should be stored in the Tools directory.
+`enabled` | `boolean`<br>Whether the plugin should be enabled or not
 
+## BundledToolConfig
+
+Configuration for a tool bundled alongsize the server
+
+Name | Description
+---- | -----------
+`version` | `string`<br>Version string for the current tool data
+`refName` | `string`<br>Ref name in the tools directory
+`dataDir` | `string`<br>Directory containing blob data for this tool. If empty, the tools/{id} folder next to the server will be used.
+`id` | `string`<br>Unique identifier for the tool
+`name` | `string`<br>Name of the tool
+`description` | `string`<br>Description for the tool
+`category` | `string`<br>Category for the tool. Will cause the tool to be shown in a different tab in the dashboard.
+`group` | `string`<br>Grouping key for different variations of the same tool. The dashboard will show these together.
+`platforms` | `string[]`<br>Platforms for this tool. Takes the form of a NET RID (https://learn.microsoft.com/en-us/dotnet/core/rid-catalog).
+`public` | `boolean`<br>Whether this tool should be exposed for download on a public endpoint without authentication
+`showInUgs` | `boolean`<br>Whether to show this tool for download in the UGS tools menu
+`showInDashboard` | `boolean`<br>Whether to show this tool for download in the dashboard
+`showInToolbox` | `boolean`<br>Whether to show this tool for download in Unreal Toolbox
+`metadata` | `string` `->` `string`<br>Metadata for this tool
+`namespaceId` | `string`<br>Default namespace for new deployments of this tool
+`maxDeployments` | `integer`<br>Maximum number of deployments to retain for this tool. When exceeded, the oldest deployments are removed.
+`acl` | [AclConfig](#aclconfig)<br>Permissions for the tool
+
+## AclConfig
+
+Parameters to update an ACL
+
+Name | Description
+---- | -----------
+`entries` | [AclEntryConfig](#aclentryconfig)`[]`<br>Entries to replace the existing ACL
+`profiles` | [AclProfileConfig](#aclprofileconfig)`[]`<br>Defines profiles which allow grouping sets of actions into named collections
+`inherit` | `boolean`<br>Whether to inherit permissions from the parent ACL
+`exceptions` | `string[]`<br>List of exceptions to the inherited setting
+
+## AclEntryConfig
+
+Individual entry in an ACL
+
+Name | Description
+---- | -----------
+`claim` | [AclClaimConfig](#aclclaimconfig)<br>Name of the user or group
+`actions` | `string[]`<br>Array of actions to allow
+`profiles` | `string[]`<br>List of profiles to grant
+
+## AclClaimConfig
+
+New claim to create
+
+Name | Description
+---- | -----------
+`type` | `string`<br>The claim type
+`value` | `string`<br>The claim value
+
+## AclProfileConfig
+
+Configuration for an ACL profile. This defines a preset group of actions which can be given to a user via an ACL entry.
+
+Name | Description
+---- | -----------
+`id` | `string`<br>Identifier for this profile
+`actions` | `string[]`<br>Actions to include
+`excludeActions` | `string[]`<br>Actions to exclude from the inherited actions
+`extends` | `string[]`<br>Other profiles to extend from
+
 ## BuildServerConfig
 
@@ -364,38 +426,50 @@
 `maxChildCount` | `integer`<br>Maximum number of children in each node
 `sliceThreshold` | `integer`<br>Threshold hash value for splitting interior nodes
 
-## ExamplePluginServerConfig
+## ComputeServerConfig
 
-Server configuration for the ExamplePlugin
+Static configuration for the compute plugin
 
 Name | Description
 ---- | -----------
-`enabled` | `boolean`<br>Whether the plugin should be enabled or not
-
-## SecretsServerConfig
-
-Static configuration for the secrets plugin
-
-Name | Description
----- | -----------
+`enableUpgradeTasks` | `boolean`<br>Whether to enable the upgrade task source, always upgrading agents to the latest version
 `withAws` | `boolean`<br>Whether to enable Amazon Web Services (AWS) specific features
+`awsRegions` | `string[]`<br>List of AWS regions for Horde to be aware of (e.g. us-east-1 or eu-central-1) Right now, this is only used for replicating CloudWatch metrics to multiple regions
+`awsAutoScalingQueueUrls` | `string[]`<br>AWS SQS queue URLs where lifecycle events from EC2 auto-scaling are received
+`fleetManagerV2` | [FleetManagerType](#fleetmanagertype-enum)<br>Default fleet manager to use (when not specified by pool)
+`fleetManagerV2Config` | `object`<br>Config for the fleet manager (serialized JSON)
+`autoEnrollAgents` | `boolean`<br>Whether to automatically enroll agents in the farm
+`defaultAgentPoolSizeStrategy` | [PoolSizeStrategy](#poolsizestrategy-enum)<br>Default agent pool sizing strategy for pools that doesn't have one explicitly configured
+`agentPoolScaleOutCooldownSeconds` | `integer`<br>Scale-out cooldown for auto-scaling agent pools (in seconds). Can be overridden by per-pool settings.
+`agentPoolScaleInCooldownSeconds` | `integer`<br>Scale-in cooldown for auto-scaling agent pools (in seconds). Can be overridden by per-pool settings.
+`computeTunnelPort` | `integer`<br>Port to listen on for tunneling compute sockets to agents
+`computeTunnelAddress` | `string`<br>What address (host:port) clients should connect to for compute socket tunneling Port may differ from  if Horde server is behind a reverse proxy/firewall
 `enabled` | `boolean`<br>Whether the plugin should be enabled or not
 
-## ExperimentalServerConfig
+## FleetManagerType (Enum)
 
-Server configuration for the experimental plugin
+Available fleet managers
 
 Name | Description
 ---- | -----------
-`enabled` | `boolean`<br>Whether the plugin should be enabled or not
+`Default` | Default fleet manager
+`NoOp` | No-op fleet manager.
+`Aws` | Fleet manager for handling AWS EC2 instances. Will create and/or terminate instances from scratch.
+`AwsReuse` | Fleet manager for handling AWS EC2 instances. Will start already existing but stopped instances to reuse existing EBS disks.
+`AwsRecycle` | Fleet manager for handling AWS EC2 instances. Will start already existing but stopped instances to reuse existing EBS disks.
+`AwsAsg` | Fleet manager for handling AWS EC2 instances. Uses an EC2 auto-scaling group for controlling the number of running instances.
 
-## PluginServerConfig
+## PoolSizeStrategy (Enum)
 
-Base class for plugin server config objects
+Available pool sizing strategies
 
 Name | Description
 ---- | -----------
-`enabled` | `boolean`<br>Whether the plugin should be enabled or not
+`LeaseUtilization` | Strategy based on lease utilization
+`JobQueue` | Strategy based on size of job build queue
+`NoOp` | No-op strategy used as fallback/default behavior
+`ComputeQueueAwsMetric` | A no-op strategy that reports metrics to let an external AWS auto-scaling policy scale the fleet
+`LeaseUtilizationAwsMetric` | A no-op strategy that reports metrics to let an external AWS auto-scaling policy scale the fleet
 
 ## StorageServerConfig
 
@@ -455,77 +529,3 @@
 `Profile` | Read credentials from the  profile in the AWS config file
 `AssumeRole` | Assume a particular role. Should specify ARN in
 `AssumeRoleWebIdentity` | Assume a particular role using the current environment variables.
-
-## ToolsServerConfig
-
-Server configuration for bundled tools
-
-Name | Description
----- | -----------
-`bundledTools` | [BundledToolConfig](#bundledtoolconfig)`[]`<br>Tools bundled along with the server. Data for each tool can be produced using the 'bundle create' command, and should be stored in the Tools directory.
-`enabled` | `boolean`<br>Whether the plugin should be enabled or not
-
-## BundledToolConfig
-
-Configuration for a tool bundled alongsize the server
-
-Name | Description
----- | -----------
-`version` | `string`<br>Version string for the current tool data
-`refName` | `string`<br>Ref name in the tools directory
-`dataDir` | `string`<br>Directory containing blob data for this tool. If empty, the tools/{id} folder next to the server will be used.
-`id` | `string`<br>Unique identifier for the tool
-`name` | `string`<br>Name of the tool
-`description` | `string`<br>Description for the tool
-`category` | `string`<br>Category for the tool. Will cause the tool to be shown in a different tab in the dashboard.
-`group` | `string`<br>Grouping key for different variations of the same tool. The dashboard will show these together.
-`platforms` | `string[]`<br>Platforms for this tool. Takes the form of a NET RID (https://learn.microsoft.com/en-us/dotnet/core/rid-catalog).
-`public` | `boolean`<br>Whether this tool should be exposed for download on a public endpoint without authentication
-`showInUgs` | `boolean`<br>Whether to show this tool for download in the UGS tools menu
-`showInDashboard` | `boolean`<br>Whether to show this tool for download in the dashboard
-`showInToolbox` | `boolean`<br>Whether to show this tool for download in Unreal Toolbox
-`metadata` | `string` `->` `string`<br>Metadata for this tool
-`namespaceId` | `string`<br>Default namespace for new deployments of this tool
-`maxDeployments` | `integer`<br>Maximum number of deployments to retain for this tool. When exceeded, the oldest deployments are removed.
-`acl` | [AclConfig](#aclconfig)<br>Permissions for the tool
-
-## AclConfig
-
-Parameters to update an ACL
-
-Name | Description
----- | -----------
-`entries` | [AclEntryConfig](#aclentryconfig)`[]`<br>Entries to replace the existing ACL
-`profiles` | [AclProfileConfig](#aclprofileconfig)`[]`<br>Defines profiles which allow grouping sets of actions into named collections
-`inherit` | `boolean`<br>Whether to inherit permissions from the parent ACL
-`exceptions` | `string[]`<br>List of exceptions to the inherited setting
-
-## AclEntryConfig
-
-Individual entry in an ACL
-
-Name | Description
----- | -----------
-`claim` | [AclClaimConfig](#aclclaimconfig)<br>Name of the user or group
-`actions` | `string[]`<br>Array of actions to allow
-`profiles` | `string[]`<br>List of profiles to grant
-
-## AclClaimConfig
-
-New claim to create
-
-Name | Description
----- | -----------
-`type` | `string`<br>The claim type
-`value` | `string`<br>The claim value
-
-## AclProfileConfig
-
-Configuration for an ACL profile. This defines a preset group of actions which can be given to a user via an ACL entry.
-
-Name | Description
----- | -----------
-`id` | `string`<br>Identifier for this profile
-`actions` | `string[]`<br>Actions to include
-`excludeActions` | `string[]`<br>Actions to exclude from the inherited actions
-`extends` | `string[]`<br>Other profiles to extend from
