Horde: Add template filter to job search

diff --git a/Engine/Source/Programs/Horde/HordeDashboard/src/components/JobSearchSimple.tsx b/Engine/Source/Programs/Horde/HordeDashboard/src/components/JobSearchSimple.tsx
--- a/Engine/Source/Programs/Horde/HordeDashboard/src/components/JobSearchSimple.tsx
+++ b/Engine/Source/Programs/Horde/HordeDashboard/src/components/JobSearchSimple.tsx
@@ -1,9 +1,9 @@
 // Copyright Epic Games, Inc. All Rights Reserved.
 
-import { Checkbox, ContextualMenuItemType, DatePicker, DefaultButton, DetailsList, DetailsRow, IColumn, IconButton, IContextualMenuItem, IContextualMenuItemProps, IContextualMenuProps, IDetailsGroupRenderProps, IDetailsListProps, IGroup, Label, Modal, PrimaryButton, SelectionMode, Spinner, SpinnerSize, Stack, Text, TextField } from "@fluentui/react";
+import { Callout, Checkbox, ContextualMenuItemType, DatePicker, DefaultButton, DetailsList, DetailsRow, DirectionalHint, IColumn, IconButton, IContextualMenuItem, IContextualMenuItemProps, IContextualMenuProps, IDetailsGroupRenderProps, IDetailsListProps, IGroup, Label, Modal, PrimaryButton, SelectionMode, Spinner, SpinnerSize, Stack, Text, TextField } from "@fluentui/react";
 import { observer } from "mobx-react-lite";
 import moment from "moment";
-import { useEffect, useMemo, useState } from "react";
+import { useEffect, useMemo, useRef, useState } from "react";
 import { Link, useLocation } from "react-router-dom";
 import backend from "../backend";
 import { GetJobResponse, GetJobsTabResponse, GetJobStepRefResponse, JobData, JobQuery, TabType, GetTemplateRefResponse, GetJobsTabLabelColumnResponse, GetLabelStateResponse, LabelOutcome, LabelState, GetJobsTabColumnResponse, JobStepOutcome } from "../backend/Api";
@@ -56,6 +56,9 @@ export const JobSearchSimpleModal: React.FC<{ streamIdIn: string, changeListIn?:
    const [searchResults, setSearchResults] = useState<{items: JobItem[], groups?: IGroup[], querying: boolean}>({items: [], querying: false});
    const query = useQuery();
    const [templateData, setTemplateData] = useState<{ streamId: string; templates: GetTemplateRefResponse[] } | undefined>(undefined);
+   const [templateFilter, setTemplateFilter] = useState("");
+   const [templateSearchOpen, setTemplateSearchOpen] = useState(false);
+   const templateSearchRef = useRef<HTMLDivElement>(null);
    const location = useLocation();
    const [isFirstRender, setIsFirstRender] = useState(true);
 
@@ -478,12 +481,12 @@ export const JobSearchSimpleModal: React.FC<{ streamIdIn: string, changeListIn?:
    templateOptions.unshift({
       key: `select_all_templates`,
       onRender: () => {
-         return <Stack horizontal verticalFill verticalAlign="center" style={{ padding: "4px 12px 4px 7px", cursor: "pointer" }} onClick={(ev) => { 
+         return <Stack horizontal verticalFill verticalAlign="center" style={{ padding: "4px 12px 4px 7px", cursor: "pointer" }} onClick={(ev) => {
             ev?.stopPropagation();
             ev?.preventDefault();
 
             allTemplatesSelected = !allTemplatesSelected;
-            
+
             if (allTemplatesSelected) {
                searchState.templates = allTemplates;
                setSearchState({ ...searchState, containsParameter: undefined })
@@ -499,6 +502,7 @@ export const JobSearchSimpleModal: React.FC<{ streamIdIn: string, changeListIn?:
       }
    });
 
+
    const templateMenuProps: IContextualMenuProps = {
       shouldFocusOnMount: true,
       subMenuHoverDelay: 0,
@@ -863,7 +867,103 @@ export const JobSearchSimpleModal: React.FC<{ streamIdIn: string, changeListIn?:
 
                      <Stack>
                         <Label>{searchState.templates.size > 1 ? "Templates" : "Template"}</Label>
-                        <DefaultButton disabled={searchState.global} style={{ width: 352, textAlign: "left" }} menuProps={templateMenuProps} text={templateText} />
+                        <div ref={templateSearchRef}>
+                           <Stack horizontal tokens={{ childrenGap: 4 }}>
+                              {!templateSearchOpen && <DefaultButton disabled={searchState.global} style={{ width: 320, textAlign: "left" }} menuProps={templateMenuProps} text={templateText} />}
+                              {templateSearchOpen && <TextField
+                                 autoFocus
+                                 disabled={searchState.global}
+                                 value={templateFilter}
+                                 spellCheck={false}
+                                 autoComplete="off"
+                                 style={{ width: 320 }}
+                                 placeholder="Filter templates..."
+                                 onChange={(ev, newValue) => setTemplateFilter(newValue ?? "")}
+                                 onKeyDown={(ev) => { if (ev.key === "Escape") { setTemplateSearchOpen(false); setTemplateFilter(""); } }}
+                              />}
+                              <IconButton
+                                 iconProps={{ iconName: templateSearchOpen ? "Cancel" : "Filter" }}
+                                 title={templateSearchOpen ? "Close filter" : "Filter templates"}
+                                 disabled={searchState.global}
+                                 styles={{ root: { height: 32, width: 32 } }}
+                                 onClick={() => {
+                                    if (templateSearchOpen) {
+                                       setTemplateSearchOpen(false);
+                                       setTemplateFilter("");
+                                    } else {
+                                       setTemplateSearchOpen(true);
+                                    }
+                                 }}
+                              />
+                           </Stack>
+                        </div>
+                        {templateSearchOpen && templateFilter && templateSearchRef.current && (() => {
+                           const filter = templateFilter.toLowerCase();
+                           const filteredCategories: { category: string; templates: GetTemplateRefResponse[] }[] = [];
+                           Array.from(sorted.keys()).sort((a, b) => a < b ? -1 : 1).forEach(cat => {
+                              const catTemplates = sorted.get(cat);
+                              if (!catTemplates?.length) return;
+                              const matched = catTemplates.filter(t => t.name.toLowerCase().includes(filter) || t.id.toLowerCase().includes(filter));
+                              if (matched.length > 0) {
+                                 filteredCategories.push({ category: cat, templates: matched.sort((a, b) => a.name < b.name ? -1 : 1) });
+                              }
+                           });
+                           const totalMatches = filteredCategories.reduce((sum, c) => sum + c.templates.length, 0);
+                           return <Callout
+                              target={templateSearchRef.current}
+                              directionalHint={DirectionalHint.bottomLeftEdge}
+                              isBeakVisible={false}
+                              onDismiss={() => { setTemplateSearchOpen(false); setTemplateFilter(""); }}
+                              styles={{ root: { width: 352 } }}
+                           >
+                              <div style={{ maxHeight: 320, overflowY: "auto", padding: "4px 0" }}>
+                                 {totalMatches === 0 && <Stack style={{ padding: "8px 12px" }}><Text variant="small" styles={{ root: { color: "#888" } }}>No matching templates</Text></Stack>}
+                                 {filteredCategories.map(({ category, templates: catTemplates }) => (
+                                    <Stack key={category}>
+                                       <Stack style={{ padding: "6px 12px 2px 12px" }}>
+                                          <Text variant="smallPlus" styles={{ root: { fontWeight: 600, color: dashboard.darktheme ? "#ABABAB" : "#575757" } }}>{category}</Text>
+                                       </Stack>
+                                       {catTemplates.map(t => (
+                                          <Stack key={t.id} horizontal verticalAlign="center" style={{ padding: "4px 12px 4px 20px", cursor: "pointer" }}
+                                             onClick={(ev) => {
+                                                ev.preventDefault();
+                                                ev.stopPropagation();
+                                                setSearchState(prev => {
+                                                   const newTemplates = new Set(prev.templates);
+                                                   if (newTemplates.has(t.id)) {
+                                                      newTemplates.delete(t.id);
+                                                   } else {
+                                                      newTemplates.add(t.id);
+                                                   }
+                                                   return { ...prev, templates: newTemplates };
+                                                });
+                                             }}>
+                                             <Checkbox checked={searchState.templates.has(t.id)} />
+                                             <Text variant="small">{t.name}</Text>
+                                          </Stack>
+                                       ))}
+                                    </Stack>
+                                 ))}
+                              </div>
+                              {totalMatches > 0 && <Stack horizontal style={{ padding: "4px 12px", borderTop: "1px solid #eee" }} tokens={{ childrenGap: 12 }}>
+                                 <Text variant="tiny" styles={{ root: { color: "#888", lineHeight: "24px" } }}>{totalMatches} match{totalMatches !== 1 ? "es" : ""}</Text>
+                                 <DefaultButton text="Select All Matches" styles={{ root: { height: 24, minWidth: 0, fontSize: 11, padding: "0 8px" } }} onClick={() => {
+                                    setSearchState(prev => {
+                                       const newTemplates = new Set(prev.templates);
+                                       filteredCategories.forEach(({ templates: ct }) => ct.forEach(t => newTemplates.add(t.id)));
+                                       return { ...prev, templates: newTemplates };
+                                    });
+                                 }} />
+                                 <DefaultButton text="Deselect All" styles={{ root: { height: 24, minWidth: 0, fontSize: 11, padding: "0 8px" } }} onClick={() => {
+                                    setSearchState(prev => {
+                                       const newTemplates = new Set(prev.templates);
+                                       filteredCategories.forEach(({ templates: ct }) => ct.forEach(t => newTemplates.delete(t.id)));
+                                       return { ...prev, templates: newTemplates };
+                                    });
+                                 }} />
+                              </Stack>}
+                           </Callout>;
+                        })()}
                      </Stack>
 
                      <Stack>
