Horde: Ignore file actions of type "unknown" when unshelving changes in PerforceMaterializer

diff --git a/Engine/Source/Programs/Horde/Drivers/JobDriver/Execution/PerforceMaterializer.cs b/Engine/Source/Programs/Horde/Drivers/JobDriver/Execution/PerforceMaterializer.cs
--- a/Engine/Source/Programs/Horde/Drivers/JobDriver/Execution/PerforceMaterializer.cs
+++ b/Engine/Source/Programs/Horde/Drivers/JobDriver/Execution/PerforceMaterializer.cs
@@ -489,11 +489,13 @@ public sealed class PerforceMaterializer : IWorkspaceMaterializer
 					case FileAction.Integrate:
 						writeFiles.Add(whereRecord);
 						break;
+					case FileAction.Unknown:
+						_logger.LogWarning("Ignoring file {File} with unknown action in shelved CL {Change}", fileRecord.DepotFile, shelveChangeNum);
+						break;
 					case FileAction.None:
 					case FileAction.Purge:
 					case FileAction.Archive:
 					case FileAction.Import:
-					case FileAction.Unknown:
 					default:
 						throw new Exception($"Unknown action '{fileRecord.Action}' for shelved file {fileRecord.DepotFile}");
 				}
