Horde: Post auto-RCA results to issue Slack triage threads
Add IRcaNotifier interface and RcaResult record in HordeServer.Shared to decouple RCA notification from Slack infrastructure across plugin boundaries.
Add AutoRcaSlackEnabled workflow config flag (defaults to false) to IWorkflow, WorkflowConfig, and GetWorkflowResponse.
Implement IRcaNotifier in SlackNotificationSink with idempotent posting to existing triage threads via PostSingleMessageToThreadAsync.
Add debug logging on all early-return guard clauses to aid first-time enablement troubleshooting.
Format Slack message with AI disclaimer, suspected culprits (Swarm/Robomerge links), reason.
Register IRcaNotifier DI forwarding in BuildPlugin within existing Slack config gate.

#Synced-CL 0

diff --git a/Engine/Source/Programs/Horde/Plugins/Build/HordeServer.Build/BuildConfig.cs b/Engine/Source/Programs/Horde/Plugins/Build/HordeServer.Build/BuildConfig.cs
--- a/Engine/Source/Programs/Horde/Plugins/Build/HordeServer.Build/BuildConfig.cs
+++ b/Engine/Source/Programs/Horde/Plugins/Build/HordeServer.Build/BuildConfig.cs
@@ -111,6 +111,12 @@ namespace HordeServer
 		/// </summary>
 		public bool EnableIssueRootCauseAnalysisNotifications { get; set; }
 
+		/// <summary>
+		/// Global kill switch for posting auto-RCA results to Slack triage threads.
+		/// When false (default), no auto-RCA Slack notifications are sent regardless of per-workflow AutoRcaSlackEnabled settings.
+		/// </summary>
+		public bool EnableAutoRcaSlackNotifications { get; set; }
+
 		/// <summary>
 		/// Whether to use multi-stage LLM analysis for auto-RCA. Multi-stage is slower but potentially more accurate. Set to false for faster processing with higher throughput.
 		/// </summary>
diff --git a/Engine/Source/Programs/Horde/Plugins/Build/HordeServer.Build/BuildPlugin.cs b/Engine/Source/Programs/Horde/Plugins/Build/HordeServer.Build/BuildPlugin.cs
--- a/Engine/Source/Programs/Horde/Plugins/Build/HordeServer.Build/BuildPlugin.cs
+++ b/Engine/Source/Programs/Horde/Plugins/Build/HordeServer.Build/BuildPlugin.cs
@@ -186,6 +186,7 @@ namespace HordeServer
 				services.AddSingleton<SlackNotificationSink>();
 				services.AddSingleton<IAvatarService, SlackNotificationSink>(sp => sp.GetRequiredService<SlackNotificationSink>());
 				services.AddSingleton<INotificationSink, SlackNotificationSink>(sp => sp.GetRequiredService<SlackNotificationSink>());
+				services.AddSingleton<IRcaNotifier, SlackNotificationSink>(sp => sp.GetRequiredService<SlackNotificationSink>());
 			}
 			else
 			{
diff --git a/Engine/Source/Programs/Horde/Plugins/Build/HordeServer.Build/Issues/WorkflowConfig.cs b/Engine/Source/Programs/Horde/Plugins/Build/HordeServer.Build/Issues/WorkflowConfig.cs
--- a/Engine/Source/Programs/Horde/Plugins/Build/HordeServer.Build/Issues/WorkflowConfig.cs
+++ b/Engine/Source/Programs/Horde/Plugins/Build/HordeServer.Build/Issues/WorkflowConfig.cs
@@ -73,6 +73,9 @@ namespace HordeServer.Issues
 		/// <inheritdoc/>
 		public bool AutoRcaEnabled { get; set; }
 
+		/// <inheritdoc/>
+		public bool AutoRcaSlackEnabled { get; set; }
+
 		/// <inheritdoc/>
 		public int MaxMentions { get; set; } = 5;
 
diff --git a/Engine/Source/Programs/Horde/Plugins/Build/HordeServer.Build/Notifications/Sinks/SlackNotificationSink.cs b/Engine/Source/Programs/Horde/Plugins/Build/HordeServer.Build/Notifications/Sinks/SlackNotificationSink.cs
--- a/Engine/Source/Programs/Horde/Plugins/Build/HordeServer.Build/Notifications/Sinks/SlackNotificationSink.cs
+++ b/Engine/Source/Programs/Horde/Plugins/Build/HordeServer.Build/Notifications/Sinks/SlackNotificationSink.cs
@@ -48,12 +48,14 @@ using MongoDB.Bson.Serialization.Attributes;
 using MongoDB.Driver;
 using StackExchange.Redis;
 
+using HordeServer.RCA;
+
 namespace HordeServer.Notifications.Sinks
 {
 	/// <summary>
 	/// Maintains a connection to Slack, in order to receive socket-mode notifications of user interactions
 	/// </summary>
-	public sealed class SlackNotificationSink : IHostedService, INotificationSink, IAvatarService, IAsyncDisposable
+	public sealed class SlackNotificationSink : IHostedService, INotificationSink, IAvatarService, IRcaNotifier, IAsyncDisposable
 	{
 		const bool DefaultAllowMentions = true;
 
@@ -68,6 +70,12 @@ namespace HordeServer.Notifications.Sinks
 		// The color to use for success messages.
 		const string SuccessColor = "#4db507";
 
+		// Max attempts when looking up an existing triage thread for auto-RCA posting.
+		const int RcaThreadLookupMaxAttempts = 4;
+
+		// Base delay in seconds for exponential backoff between thread lookup retries (5s, 10s, 20s).
+		const int RcaThreadLookupBaseDelaySecs = 5;
+
 		// Badges to add color accents
 		const string YellowgBadge = ":large_yellow_square:";
 		const string OrangeBadge = ":large_orange_square:";
@@ -1632,6 +1640,118 @@ namespace HordeServer.Notifications.Sinks
 			}
 		}
 
+		/// <inheritdoc/>
+		public async Task NotifyRcaCompleteAsync(int issueId, RcaResult result, CancellationToken cancellationToken)
+		{
+			// Global kill switch check
+			if (!_buildConfig.CurrentValue.EnableAutoRcaSlackNotifications)
+			{
+				_logger.LogDebug("Skipping auto-RCA notification for issue {IssueId}: EnableAutoRcaSlackNotifications is disabled globally", issueId);
+				return;
+			}
+
+			// Get issue details to find the workflow
+			IIssueDetails? details = await _issueService.GetIssueDetailsAsync(issueId, cancellationToken);
+			if (details == null || details.Spans.Count == 0)
+			{
+				_logger.LogDebug("Skipping auto-RCA notification for issue {IssueId}: issue not found or has no spans", issueId);
+				return;
+			}
+
+			IIssueSpan span = details.Spans[0];
+			WorkflowId? workflowId = span.LastFailure.Annotations.WorkflowId;
+			if (workflowId == null)
+			{
+				_logger.LogDebug("Skipping auto-RCA notification for issue {IssueId}: no workflow annotation on span", issueId);
+				return;
+			}
+
+			BuildConfig buildConfig = _buildConfig.CurrentValue;
+			if (!buildConfig.TryGetStream(span.StreamId, out StreamConfig? streamConfig))
+			{
+				_logger.LogDebug("Skipping auto-RCA notification for issue {IssueId}: stream {StreamId} not found in config", issueId, span.StreamId);
+				return;
+			}
+			if (!streamConfig.TryGetWorkflow(workflowId.Value, out WorkflowConfig? workflow))
+			{
+				_logger.LogDebug("Skipping auto-RCA notification for issue {IssueId}: workflow {WorkflowId} not found in stream {StreamId}", issueId, workflowId.Value, span.StreamId);
+				return;
+			}
+			if (!workflow.AutoRcaSlackEnabled || String.IsNullOrEmpty(workflow.TriageChannel))
+			{
+				_logger.LogDebug("Skipping auto-RCA notification for issue {IssueId}: AutoRcaSlackEnabled={Enabled}, TriageChannel={Channel}",
+					issueId, workflow.AutoRcaSlackEnabled, workflow.TriageChannel ?? "(empty)");
+				return;
+			}
+
+			// Find the existing Slack thread (with exponential backoff in case thread creation is delayed)
+			MessageStateDocument? state = null;
+			for (int attempt = 0; attempt < RcaThreadLookupMaxAttempts && state == null; attempt++)
+			{
+				if (attempt > 0)
+				{
+					int delaySecs = RcaThreadLookupBaseDelaySecs * (1 << (attempt - 1));
+					_logger.LogDebug("Auto-RCA thread lookup for issue {IssueId}: attempt {Attempt}/{MaxAttempts}, waiting {DelaySecs}s",
+						issueId, attempt + 1, RcaThreadLookupMaxAttempts, delaySecs);
+					await Task.Delay(TimeSpan.FromSeconds(delaySecs), cancellationToken);
+				}
+				state = await GetMessageStateAsync(workflow.TriageChannel, GetTriageThreadEventId(issueId), cancellationToken);
+			}
+			if (state == null)
+			{
+				_logger.LogWarning("No Slack thread found for issue {IssueId} after {MaxAttempts} attempts, skipping auto-RCA notification",
+					issueId, RcaThreadLookupMaxAttempts);
+				return;
+			}
+
+			// Post via idempotent one-time pattern
+			string rcaEventId = $"issue_{issueId}_auto_rca";
+			string rcaMessage = FormatAutoRcaMessage(result);
+			await PostSingleMessageToThreadAsync(workflow.TriageChannel, rcaEventId, state.MessageId, rcaMessage, cancellationToken);
+		}
+
+		const int MaxRunnerUpsInSlack = 3;
+		const int MaxReasonLength = 700;
+
+		static string TruncateReason(string reason)
+			=> reason.Length > MaxReasonLength ? reason[..MaxReasonLength] + "..." : reason;
+
+		string FormatAutoRcaMessage(RcaResult result)
+		{
+			StringBuilder sb = new();
+			sb.AppendLine(":robot_face: *Auto-RCA Suggestion*");
+			sb.AppendLine("_:warning: This analysis was generated by AI and may be inaccurate. Use as a starting point for triage only._");
+
+			if (result.CulpritCommitId != null)
+			{
+				sb.AppendLine($"Suspected culprit: {FormatChange(result.CulpritCommitId)}");
+				if (!String.IsNullOrEmpty(result.Reason))
+				{
+					sb.AppendLine($"> {TruncateReason(result.Reason)}");
+				}
+			}
+
+			if (result.RunnerUps != null && result.RunnerUps.Count > 0)
+			{
+				sb.AppendLine("Other suspects:");
+				foreach (RcaRunnerUp runnerUp in result.RunnerUps.Take(MaxRunnerUpsInSlack))
+				{
+					string line = $"- {FormatChange(runnerUp.Commit)}";
+					if (!String.IsNullOrEmpty(runnerUp.BriefReason))
+					{
+						line += $": {TruncateReason(runnerUp.BriefReason)}";
+					}
+					sb.AppendLine(line);
+				}
+				if (result.RunnerUps.Count > MaxRunnerUpsInSlack)
+				{
+					sb.AppendLine($"_...and {result.RunnerUps.Count - MaxRunnerUpsInSlack} more (see dashboard)_");
+				}
+			}
+
+			return sb.ToString().TrimEnd();
+		}
+
 		async Task NotifyIssueUpdatedAsync(BuildConfig buildConfig, IUser user, IIssue issue, IIssueDetails details, CancellationToken cancellationToken)
 		{
 			string? slackUserId = await GetSlackUserIdAsync(user, cancellationToken);
diff --git a/Engine/Source/Programs/Shared/EpicGames.Horde/Issues/IWorkflow.cs b/Engine/Source/Programs/Shared/EpicGames.Horde/Issues/IWorkflow.cs
--- a/Engine/Source/Programs/Shared/EpicGames.Horde/Issues/IWorkflow.cs
+++ b/Engine/Source/Programs/Shared/EpicGames.Horde/Issues/IWorkflow.cs
@@ -101,6 +101,11 @@ namespace EpicGames.Horde.Issues
 		/// </summary>
 		bool AutoRcaEnabled { get; }
 
+		/// <summary>
+		/// Whether to post auto-RCA results to the issue's Slack triage thread
+		/// </summary>
+		bool AutoRcaSlackEnabled { get; }
+
 		/// <summary>
 		/// Maximum number of people to mention on a triage thread
 		/// </summary>
diff --git a/Engine/Source/Programs/Shared/EpicGames.Horde/Issues/WorkflowMessages.cs b/Engine/Source/Programs/Shared/EpicGames.Horde/Issues/WorkflowMessages.cs
--- a/Engine/Source/Programs/Shared/EpicGames.Horde/Issues/WorkflowMessages.cs
+++ b/Engine/Source/Programs/Shared/EpicGames.Horde/Issues/WorkflowMessages.cs
@@ -73,6 +73,9 @@ namespace EpicGames.Horde.Issues
 		/// <inheritdoc/>
 		public bool AutoRcaEnabled { get; set; }
 
+		/// <inheritdoc/>
+		public bool AutoRcaSlackEnabled { get; set; }
+
 		/// <inheritdoc/>
 		public int MaxMentions { get; set; }
 
@@ -133,6 +136,7 @@ namespace EpicGames.Horde.Issues
 			EscalateTimes = [.. workflow.EscalateTimes];
 			EscalateRootCauseAnalysisTimes = [.. workflow.EscalateRootCauseAnalysisTimes];
 			AutoRcaEnabled = workflow.AutoRcaEnabled;
+			AutoRcaSlackEnabled = workflow.AutoRcaSlackEnabled;
 			MaxMentions = workflow.MaxMentions;
 			AllowMentions = workflow.AllowMentions;
 			InviteRestrictedUsers = workflow.InviteRestrictedUsers;
