Add missing IRCANotifier.cs file, left out from CL 51907783

#rnx
#Synced-CL 0

diff --git a/Engine/Source/Programs/Horde/HordeServer.Shared/RCA/IRcaNotifier.cs b/Engine/Source/Programs/Horde/HordeServer.Shared/RCA/IRcaNotifier.cs
new file mode 100644
--- /dev/null
+++ b/Engine/Source/Programs/Horde/HordeServer.Shared/RCA/IRcaNotifier.cs
@@ -0,0 +1,31 @@
+// Copyright Epic Games, Inc. All Rights Reserved.
+
+using System.Collections.Generic;
+using EpicGames.Horde.Commits;
+
+namespace HordeServer.RCA
+{
+	/// <summary>
+	/// A runner-up suspect from an RCA analysis
+	/// </summary>
+	public record RcaRunnerUp(CommitId Commit, string? BriefReason);
+
+	/// <summary>
+	/// Result of an auto-RCA analysis
+	/// </summary>
+	public record RcaResult(CommitId? CulpritCommitId, string? Reason, double Confidence, string? Summary, string? Category, IReadOnlyList<RcaRunnerUp>? RunnerUps);
+
+	/// <summary>
+	/// Interface for notifying about completed RCA results (e.g., posting to Slack threads).
+	/// </summary>
+	public interface IRcaNotifier
+	{
+		/// <summary>
+		/// Notify that RCA analysis has completed for an issue
+		/// </summary>
+		/// <param name="issueId">The issue ID that was analyzed</param>
+		/// <param name="result">The RCA result to notify about</param>
+		/// <param name="cancellationToken">Cancellation token</param>
+		Task NotifyRcaCompleteAsync(int issueId, RcaResult result, CancellationToken cancellationToken);
+	}
+}
