Remove unnecessary using System.Collections.Generic directive

#rnx

diff --git a/Engine/Source/Programs/Horde/HordeServer.Shared/RCA/IRcaNotifier.cs b/Engine/Source/Programs/Horde/HordeServer.Shared/RCA/IRcaNotifier.cs
--- a/Engine/Source/Programs/Horde/HordeServer.Shared/RCA/IRcaNotifier.cs
+++ b/Engine/Source/Programs/Horde/HordeServer.Shared/RCA/IRcaNotifier.cs
@@ -1,6 +1,5 @@
 // Copyright Epic Games, Inc. All Rights Reserved.
 
-using System.Collections.Generic;
 using EpicGames.Horde.Commits;
 
 namespace HordeServer.RCA
