Horde: Update TargetFramework to net10.0

#rnx
#Synced-CL 51916670

diff --git a/Engine/Source/Programs/Horde/Drivers/JobDriver/JobDriver.csproj b/Engine/Source/Programs/Horde/Drivers/JobDriver/JobDriver.csproj
--- a/Engine/Source/Programs/Horde/Drivers/JobDriver/JobDriver.csproj
+++ b/Engine/Source/Programs/Horde/Drivers/JobDriver/JobDriver.csproj
@@ -72,7 +72,6 @@
     <PackageReference Include="Serilog.Sinks.File" Version="5.0.0" />
     <PackageReference Include="System.Diagnostics.PerformanceCounter" Version="8.0.1" />
     <PackageReference Include="System.ServiceProcess.ServiceController" Version="8.0.1" />
-    <PackageReference Include="System.Text.Json" Version="8.0.6" />
   </ItemGroup>
   
   <ItemGroup>
diff --git a/Engine/Source/Programs/Horde/HordeAgent.Tests/Utility/WindowsProcessMemoryTrackerTests.cs b/Engine/Source/Programs/Horde/HordeAgent.Tests/Utility/WindowsProcessMemoryTrackerTests.cs
--- a/Engine/Source/Programs/Horde/HordeAgent.Tests/Utility/WindowsProcessMemoryTrackerTests.cs
+++ b/Engine/Source/Programs/Horde/HordeAgent.Tests/Utility/WindowsProcessMemoryTrackerTests.cs
@@ -233,7 +233,7 @@ public class WindowsProcessMemoryTrackerTests
 		}
 
 		// Verify the MaxCommandLineLength constant is reasonable
-		Assert.AreEqual(100, WindowsProcessMemoryTracker.MaxCommandLineLength);
+		Assert.AreEqual(500, WindowsProcessMemoryTracker.MaxCommandLineLength);
 	}
 
 	[TestMethod]
diff --git a/Engine/Source/Programs/Horde/HordeAgent/Commands/Utilities/CreateCertCommand.cs b/Engine/Source/Programs/Horde/HordeAgent/Commands/Utilities/CreateCertCommand.cs
--- a/Engine/Source/Programs/Horde/HordeAgent/Commands/Utilities/CreateCertCommand.cs
+++ b/Engine/Source/Programs/Horde/HordeAgent/Commands/Utilities/CreateCertCommand.cs
@@ -72,7 +72,7 @@ namespace HordeAgent.Commands.Certs
 			logger.LogInformation("Writing private cert: {PrivateCert}", new FileReference(PrivateCertFile).FullName);
 			await File.WriteAllBytesAsync(PrivateCertFile, privateCertData);
 
-			using X509Certificate2 certificate = new X509Certificate2(privateCertData);
+			using X509Certificate2 certificate = X509CertificateLoader.LoadPkcs12(privateCertData, null);
 			logger.LogInformation("Certificate thumbprint is {Thumbprint}", certificate.Thumbprint);
 			logger.LogInformation("Add this thumbprint to list of trusted servers in appsettings.json to trust this server.");
 
diff --git a/Engine/Source/Programs/Horde/HordeAgent/HordeAgent.csproj b/Engine/Source/Programs/Horde/HordeAgent/HordeAgent.csproj
--- a/Engine/Source/Programs/Horde/HordeAgent/HordeAgent.csproj
+++ b/Engine/Source/Programs/Horde/HordeAgent/HordeAgent.csproj
@@ -60,7 +60,6 @@
     <PackageReference Include="Serilog.Sinks.File" Version="5.0.0" />
     <PackageReference Include="System.Diagnostics.PerformanceCounter" Version="8.0.1" />
     <PackageReference Include="System.ServiceProcess.ServiceController" Version="8.0.1" />
-    <PackageReference Include="System.Text.Json" Version="8.0.6" />
   </ItemGroup>
 
   <ItemGroup>
diff --git a/Engine/Source/Programs/Horde/HordeAgent/Leases/Handlers/UpgradeHandler.cs b/Engine/Source/Programs/Horde/HordeAgent/Leases/Handlers/UpgradeHandler.cs
--- a/Engine/Source/Programs/Horde/HordeAgent/Leases/Handlers/UpgradeHandler.cs
+++ b/Engine/Source/Programs/Horde/HordeAgent/Leases/Handlers/UpgradeHandler.cs
@@ -56,7 +56,7 @@ namespace HordeAgent.Leases.Handlers
 				// Extract it to a temporary directory
 				DirectoryReference extractedDir = DirectoryReference.Combine(upgradeDir, "Extracted");
 				DirectoryReference.CreateDirectory(extractedDir);
-				ZipFile.ExtractToDirectory(outputFile.FullName, extractedDir.FullName);
+				await ZipFile.ExtractToDirectoryAsync(outputFile.FullName, extractedDir.FullName, cancellationToken);
 
 				//				// Debug code for updating an agent with the local version
 				//				foreach (FileInfo SourceFile in new FileInfo(Assembly.GetExecutingAssembly().Location).Directory.EnumerateFiles())
diff --git a/Engine/Source/Programs/Horde/HordeAgent/Leases/LeaseManager.cs b/Engine/Source/Programs/Horde/HordeAgent/Leases/LeaseManager.cs
--- a/Engine/Source/Programs/Horde/HordeAgent/Leases/LeaseManager.cs
+++ b/Engine/Source/Programs/Horde/HordeAgent/Leases/LeaseManager.cs
@@ -57,7 +57,7 @@ namespace HordeAgent.Leases
 		/// <summary>
 		/// Object used for controlling access to the access tokens and active sessions list
 		/// </summary>
-		readonly object _lockObject = new object();
+		readonly Lock _lockObject = new Lock();
 
 		/// <summary>
 		/// The list of active leases.
diff --git a/Engine/Source/Programs/Horde/HordeAgent/Services/ComputeListenerService.cs b/Engine/Source/Programs/Horde/HordeAgent/Services/ComputeListenerService.cs
--- a/Engine/Source/Programs/Horde/HordeAgent/Services/ComputeListenerService.cs
+++ b/Engine/Source/Programs/Horde/HordeAgent/Services/ComputeListenerService.cs
@@ -50,7 +50,7 @@ namespace HordeAgent.Services
 		Task _serverTask = Task.CompletedTask;
 		readonly CancellationTokenSource _cancellationSource = new CancellationTokenSource();
 
-		readonly object _lockObject = new object();
+		readonly Lock _lockObject = new Lock();
 		readonly Dictionary<ByteString, WaitingClient> _waitingClients = new Dictionary<ByteString, WaitingClient>();
 		readonly Dictionary<ByteString, TaskCompletionSource<TcpClient?>> _waitingLeases = new Dictionary<ByteString, TaskCompletionSource<TcpClient?>>();
 
diff --git a/Engine/Source/Programs/Horde/HordeAgent/Utility/VersionInfo.cs.template b/Engine/Source/Programs/Horde/HordeAgent/Utility/VersionInfo.cs.template
--- a/Engine/Source/Programs/Horde/HordeAgent/Utility/VersionInfo.cs.template
+++ b/Engine/Source/Programs/Horde/HordeAgent/Utility/VersionInfo.cs.template
@@ -24,17 +24,11 @@ public static class VersionInfo
 	/// <summary>
 	/// Version of this Horde Agent
 	/// </summary>
-	public static string Version
-	{
-		get => String.IsNullOrEmpty(VersionConstant) ? "unknown" : VersionConstant;
-	}
+	public static string Version => String.IsNullOrEmpty(VersionConstant) ? "unknown" : VersionConstant;
 	
 	/// <summary>
 	/// Self-reported tool ID of this agent. This ID will override and define exact variant of agent software any server-initiated upgrades will use.
 	/// </summary>
 	[SuppressMessage("Maintainability", "CA1508:Avoid dead conditional code", Justification = "Generated code")]
-	public static string? ToolId
-	{
-		get => ToolIdConstant == "unknown" || String.IsNullOrEmpty(ToolIdConstant) ? null : ToolIdConstant;
-	}
+	public static string? ToolId => ToolIdConstant == "unknown" || String.IsNullOrEmpty(ToolIdConstant) ? null : ToolIdConstant;
 }
\ No newline at end of file
diff --git a/Engine/Source/Programs/Horde/HordeCmd/Commands/Tools/ToolDownload.cs b/Engine/Source/Programs/Horde/HordeCmd/Commands/Tools/ToolDownload.cs
--- a/Engine/Source/Programs/Horde/HordeCmd/Commands/Tools/ToolDownload.cs
+++ b/Engine/Source/Programs/Horde/HordeCmd/Commands/Tools/ToolDownload.cs
@@ -37,7 +37,7 @@ class ToolDownload(IHordeClient hordeClient) : Command
 			await stream.CopyToAsync(zipOutputStream);
 		}
 
-		ZipFile.ExtractToDirectory(zipFile.FullName, OutputDir.FullName, true);
+		await ZipFile.ExtractToDirectoryAsync(zipFile.FullName, OutputDir.FullName, true, CancellationToken.None);
 		return 0;
 	}
 }
diff --git a/Engine/Source/Programs/Horde/HordeCmd/HordeCmd.csproj b/Engine/Source/Programs/Horde/HordeCmd/HordeCmd.csproj
--- a/Engine/Source/Programs/Horde/HordeCmd/HordeCmd.csproj
+++ b/Engine/Source/Programs/Horde/HordeCmd/HordeCmd.csproj
@@ -21,7 +21,6 @@
     <PackageReference Include="Serilog.Settings.Configuration" Version="8.0.4" />
     <PackageReference Include="Serilog.Sinks.Console" Version="5.0.1" />
     <PackageReference Include="Serilog.Sinks.File" Version="5.0.0" />
-    <PackageReference Include="System.Text.Json" Version="8.0.6" />
   </ItemGroup>
 
   <ItemGroup>
diff --git a/Engine/Source/Programs/Horde/HordeServer.Tests.Shared/HordeServer.Tests.Shared.csproj b/Engine/Source/Programs/Horde/HordeServer.Tests.Shared/HordeServer.Tests.Shared.csproj
--- a/Engine/Source/Programs/Horde/HordeServer.Tests.Shared/HordeServer.Tests.Shared.csproj
+++ b/Engine/Source/Programs/Horde/HordeServer.Tests.Shared/HordeServer.Tests.Shared.csproj
@@ -10,9 +10,8 @@
   </ItemGroup>
 
   <ItemGroup>
-    <PackageReference Include="Microsoft.AspNetCore.Mvc.Testing" Version="8.0.24" />
+    <PackageReference Include="Microsoft.AspNetCore.Mvc.Testing" Version="10.0.3" />
     <PackageReference Include="Moq" Version="4.20.72" />
-    <PackageReference Include="Microsoft.Extensions.Caching.Memory" Version="8.0.1" />
   </ItemGroup>
 
 </Project>
diff --git a/Engine/Source/Programs/Horde/HordeServer.Tests/HordeServer.Tests.csproj b/Engine/Source/Programs/Horde/HordeServer.Tests/HordeServer.Tests.csproj
--- a/Engine/Source/Programs/Horde/HordeServer.Tests/HordeServer.Tests.csproj
+++ b/Engine/Source/Programs/Horde/HordeServer.Tests/HordeServer.Tests.csproj
@@ -6,9 +6,8 @@
   </PropertyGroup>
 
   <ItemGroup>
-    <PackageReference Include="Microsoft.AspNetCore.Mvc.Testing" Version="8.0.24" />
+    <PackageReference Include="Microsoft.AspNetCore.Mvc.Testing" Version="10.0.3" />
     <PackageReference Include="Moq" Version="4.20.72" />
-    <PackageReference Include="Microsoft.Extensions.Caching.Memory" Version="8.0.1" />
   </ItemGroup>
 
   <ItemGroup>
diff --git a/Engine/Source/Programs/Horde/HordeServer.Tests/ProgramTests.cs b/Engine/Source/Programs/Horde/HordeServer.Tests/ProgramTests.cs
--- a/Engine/Source/Programs/Horde/HordeServer.Tests/ProgramTests.cs
+++ b/Engine/Source/Programs/Horde/HordeServer.Tests/ProgramTests.cs
@@ -2,7 +2,6 @@
 
 using System;
 using System.IO;
-using System.Runtime.InteropServices;
 using System.Security.Cryptography.X509Certificates;
 using EpicGames.Core;
 using HordeServer.Commands;
@@ -16,36 +15,42 @@ namespace HordeServer.Tests
 		[TestMethod]
 		public void TestReadGrpcCertificate()
 		{
-			if (!RuntimeInformation.IsOSPlatform(OSPlatform.Windows))
-			{
-				// Test only works on Windows as .FriendlyName cannot be set on Unix.
-				return;
-			}
-
 			string friendlyName = "A testing cert";
 			byte[] tempCertData = CertificateUtils.CreateSelfSignedCert("testing.epicgames.com", friendlyName);
 			string tempCertPath = Path.GetTempFileName();
-			File.WriteAllBytes(tempCertPath, tempCertData);
+			try
+			{
+				File.WriteAllBytes(tempCertPath, tempCertData);
+				using X509Certificate2 originalCert = X509CertificateLoader.LoadPkcs12(tempCertData, null);
+				string expectedThumbprint = originalCert.Thumbprint;
 
-			// No cert given
-			Assert.IsNull(ServerCommand.ReadGrpcCertificate(new() { ServerPrivateCert = null }));
+				// No cert given
+				Assert.IsNull(ServerCommand.ReadGrpcCertificate(new() { ServerPrivateCert = null }));
 
-			// Cert as file path
-			{
-				X509Certificate2? cert = ServerCommand.ReadGrpcCertificate(new() { ServerPrivateCert = tempCertPath });
-				Assert.IsNotNull(cert);
-				Assert.AreEqual(friendlyName, cert!.FriendlyName);
-			}
+				// Cert as file path
+				{
+					using X509Certificate2? cert = ServerCommand.ReadGrpcCertificate(new() { ServerPrivateCert = tempCertPath });
+					Assert.IsNotNull(cert);
+					Assert.AreEqual(expectedThumbprint, cert!.Thumbprint);
+				}
 
-			// Cert as base64 data
+				// Cert as base64 data
+				{
+					string tempCertBase64 = Convert.ToBase64String(tempCertData);
+					using X509Certificate2? cert = ServerCommand.ReadGrpcCertificate(new() { ServerPrivateCert = "base64:" + tempCertBase64 });
+					Assert.IsNotNull(cert);
+					Assert.AreEqual(expectedThumbprint, cert!.Thumbprint);
+				}
+			}
+			finally
 			{
-				string tempCertBase64 = Convert.ToBase64String(tempCertData);
-				X509Certificate2? cert = ServerCommand.ReadGrpcCertificate(new() { ServerPrivateCert = "base64:" + tempCertBase64 });
-				Assert.IsNotNull(cert);
-				Assert.AreEqual(friendlyName, cert!.FriendlyName);
+				try
+				{
+					File.Delete(tempCertPath);
+				}
+				catch (IOException) { }
+				catch (UnauthorizedAccessException) { }
 			}
-
-			File.Delete(tempCertPath);
 		}
 	}
 }
diff --git a/Engine/Source/Programs/Horde/HordeServer/Accounts/AccountCollection.cs b/Engine/Source/Programs/Horde/HordeServer/Accounts/AccountCollection.cs
--- a/Engine/Source/Programs/Horde/HordeServer/Accounts/AccountCollection.cs
+++ b/Engine/Source/Programs/Horde/HordeServer/Accounts/AccountCollection.cs
@@ -46,8 +46,7 @@ namespace HordeServer.Accounts
 		/// <returns></returns>
 		public static byte[] HashPassword(string password, byte[] salt)
 		{
-			using Rfc2898DeriveBytes rfc2898DeriveBytes = new(password, salt, Iterations, HashAlgorithmName.SHA256);
-			return rfc2898DeriveBytes.GetBytes(KeySize);
+			return Rfc2898DeriveBytes.Pbkdf2(password, salt, Iterations, HashAlgorithmName.SHA256, KeySize);
 		}
 
 		/// <summary>
@@ -60,7 +59,7 @@ namespace HordeServer.Accounts
 		public static bool ValidatePassword(string password, byte[] salt, byte[] correctHash)
 		{
 			byte[] hash = HashPassword(password, salt);
-			return hash.SequenceEqual(correctHash);
+			return CryptographicOperations.FixedTimeEquals(hash, correctHash);
 		}
 
 		/// <summary>
diff --git a/Engine/Source/Programs/Horde/HordeServer/Commands/ServerCommand.cs b/Engine/Source/Programs/Horde/HordeServer/Commands/ServerCommand.cs
--- a/Engine/Source/Programs/Horde/HordeServer/Commands/ServerCommand.cs
+++ b/Engine/Source/Programs/Horde/HordeServer/Commands/ServerCommand.cs
@@ -4,6 +4,7 @@ using System;
 using System.IO;
 using System.Linq;
 using System.Runtime.Versioning;
+using System.Security.Cryptography;
 using System.Security.Cryptography.X509Certificates;
 using System.Threading;
 using System.Threading.Tasks;
@@ -219,8 +220,15 @@ namespace HordeServer.Commands
 			}
 			else if (hordeSettings.ServerPrivateCert.StartsWith(base64Prefix, StringComparison.Ordinal))
 			{
-				byte[] certData = Convert.FromBase64String(hordeSettings.ServerPrivateCert.Replace(base64Prefix, "", StringComparison.Ordinal));
-				return new X509Certificate2(certData);
+				try
+				{
+					byte[] certData = Convert.FromBase64String(hordeSettings.ServerPrivateCert.Substring(base64Prefix.Length));
+					return X509CertificateLoader.LoadPkcs12(certData, null);
+				}
+				catch (Exception ex) when (ex is FormatException or CryptographicException)
+				{
+					throw new InvalidOperationException($"Failed to load PKCS#12 certificate from base64 value in setting '{nameof(ServerSettings.ServerPrivateCert)}': {ex.Message}", ex);
+				}
 			}
 			else
 			{
@@ -234,7 +242,14 @@ namespace HordeServer.Commands
 					serverPrivateCert = new FileReference(hordeSettings.ServerPrivateCert);
 				}
 
-				return new X509Certificate2(FileReference.ReadAllBytes(serverPrivateCert));
+				try
+				{
+					return X509CertificateLoader.LoadPkcs12(FileReference.ReadAllBytes(serverPrivateCert), null);
+				}
+				catch (Exception ex) when (ex is IOException or UnauthorizedAccessException or CryptographicException)
+				{
+					throw new InvalidOperationException($"Failed to load PKCS#12 certificate from path '{serverPrivateCert}' (setting '{nameof(ServerSettings.ServerPrivateCert)}'): {ex.Message}", ex);
+				}
 			}
 		}
 
diff --git a/Engine/Source/Programs/Horde/HordeServer/Dockerfile b/Engine/Source/Programs/Horde/HordeServer/Dockerfile
--- a/Engine/Source/Programs/Horde/HordeServer/Dockerfile
+++ b/Engine/Source/Programs/Horde/HordeServer/Dockerfile
@@ -12,7 +12,7 @@ ARG code_coverage="false"
 FROM redis:7.2 AS redis
 
 # Build image
-FROM mcr.microsoft.com/dotnet/sdk:8.0 AS build-env
+FROM mcr.microsoft.com/dotnet/sdk:10.0 AS build-env
 WORKDIR /app
 COPY Source/Programs/Horde/nuget.config ./Source/Programs/Horde/nuget.config
 RUN DEBIAN_FRONTEND=noninteractive \
@@ -68,7 +68,7 @@ FROM scratch AS build-env-output
 COPY --from=build-env /tmp/empty /tmp/dotcover-report/*.zip /
 
 # Create the runtime image
-FROM mcr.microsoft.com/dotnet/aspnet:8.0 AS runtime
+FROM mcr.microsoft.com/dotnet/aspnet:10.0 AS runtime
 WORKDIR /app
 ENV ASPNETCORE_URLS=""
 ENV ASPNETCORE_HTTP_PORTS=""
diff --git a/Engine/Source/Programs/Horde/HordeServer/HordeServer.csproj b/Engine/Source/Programs/Horde/HordeServer/HordeServer.csproj
--- a/Engine/Source/Programs/Horde/HordeServer/HordeServer.csproj
+++ b/Engine/Source/Programs/Horde/HordeServer/HordeServer.csproj
@@ -30,17 +30,13 @@
       <PrivateAssets>all</PrivateAssets>
       <IncludeAssets>runtime; build; native; contentfiles; analyzers; buildtransitive</IncludeAssets>
     </PackageReference>
-    <PackageReference Include="Microsoft.AspNetCore.Authentication.Cookies" Version="2.3.0" />
     <PackageReference Include="Microsoft.AspNetCore.Authentication.JwtBearer" Version="8.0.20" />
     <PackageReference Include="Microsoft.AspNetCore.Authentication.OpenIdConnect" Version="8.0.20" />
     <PackageReference Include="Microsoft.AspNetCore.DataProtection.StackExchangeRedis" Version="8.0.20" />
     <PackageReference Include="Microsoft.AspNetCore.Mvc.Razor.RuntimeCompilation" Version="8.0.20" />
     <PackageReference Include="Microsoft.AspNetCore.SpaServices.Extensions" Version="8.0.20" />
-    <PackageReference Include="Microsoft.Extensions.Caching.Memory" Version="8.0.1" />
     <PackageReference Include="Microsoft.Extensions.Hosting.WindowsServices" Version="8.0.1" />
-    <PackageReference Include="Microsoft.Extensions.Http" Version="8.0.1" />
     <PackageReference Include="Microsoft.Extensions.Http.Polly" Version="8.0.20" />
-    <PackageReference Include="Microsoft.Extensions.Logging" Version="8.0.1" />
     <PackageReference Include="Microsoft.VisualStudio.Threading.Analyzers" Version="17.14.15">
       <PrivateAssets>all</PrivateAssets>
       <IncludeAssets>runtime; build; native; contentfiles; analyzers; buildtransitive</IncludeAssets>
@@ -65,8 +61,6 @@
     <PackageReference Include="SharpZipLib" Version="1.4.2" />
     <PackageReference Include="StackExchange.Redis" Version="2.8.58" />
     <PackageReference Include="Swashbuckle.AspNetCore" Version="6.5.0" />
-    <PackageReference Include="System.Linq.Async" Version="6.0.3" />
-    <PackageReference Include="System.Text.Json" Version="8.0.6" />
     <PackageReference Include="System.Data.Odbc" Version="9.0.10" />
     <PackageReference Include="T-Digest.NET" Version="1.0.3" />
     <PackageReference Include="TimeZoneConverter" Version="6.1.0" />
diff --git a/Engine/Source/Programs/Horde/HordeServer/Server/MongoService.cs b/Engine/Source/Programs/Horde/HordeServer/Server/MongoService.cs
--- a/Engine/Source/Programs/Horde/HordeServer/Server/MongoService.cs
+++ b/Engine/Source/Programs/Horde/HordeServer/Server/MongoService.cs
@@ -480,7 +480,7 @@ namespace HordeServer.Server
 				}
 
 				string certificateText = text.Substring(offset, nextOffset - offset);
-				collection.Add(new X509Certificate2(Encoding.UTF8.GetBytes(certificateText)));
+				collection.Add(X509CertificateLoader.LoadCertificate(Encoding.UTF8.GetBytes(certificateText)));
 
 				offset = nextOffset;
 			}
diff --git a/Engine/Source/Programs/Horde/HordeServer/Startup.cs b/Engine/Source/Programs/Horde/HordeServer/Startup.cs
--- a/Engine/Source/Programs/Horde/HordeServer/Startup.cs
+++ b/Engine/Source/Programs/Horde/HordeServer/Startup.cs
@@ -427,6 +427,14 @@ namespace HordeServer
 					options.Cookie.Name = CookieAuthenticationDefaults.AuthenticationScheme;
 					options.LoginPath = "/account/login/horde";
 					options.LogoutPath = "/";
+
+					// .NET 10 breaking change: cookie auth no longer redirects to login page for [ApiController] endpoints by default.
+					// AccountController.Login() is a web navigation endpoint that must redirect to the login form, so restore the redirect behavior.
+					options.Events.OnRedirectToLogin = context =>
+					{
+						context.Response.Redirect(context.RedirectUri);
+						return Task.CompletedTask;
+					};
 				}
 
 				options.Events.OnValidatePrincipal = context =>
@@ -566,7 +574,7 @@ namespace HordeServer
 			{
 				options.ForwardedHeaders = ForwardedHeaders.XForwardedFor | ForwardedHeaders.XForwardedProto;
 				options.KnownProxies.Clear();
-				options.KnownNetworks.Clear();
+				options.KnownIPNetworks.Clear();
 			});
 
 			IMvcBuilder mvcBuilder = services.AddMvc()
diff --git a/Engine/Source/Programs/Horde/Plugins/Analytics/HordeServer.Analytics/Telemetry/Metrics/MetricCollection.cs b/Engine/Source/Programs/Horde/Plugins/Analytics/HordeServer.Analytics/Telemetry/Metrics/MetricCollection.cs
--- a/Engine/Source/Programs/Horde/Plugins/Analytics/HordeServer.Analytics/Telemetry/Metrics/MetricCollection.cs
+++ b/Engine/Source/Programs/Horde/Plugins/Analytics/HordeServer.Analytics/Telemetry/Metrics/MetricCollection.cs
@@ -83,7 +83,7 @@ namespace HordeServer.Telemetry.Metrics
 				=> $"{StoreId}:{MetricId}";
 		}
 
-		readonly object _lockObject = new object();
+		readonly Lock _lockObject = new Lock();
 		readonly IMongoCollection<MetricDocument> _metrics;
 		readonly AsyncEvent _newDataEvent = new AsyncEvent();
 		readonly AsyncEvent _flushEvent = new AsyncEvent();
diff --git a/Engine/Source/Programs/Horde/Plugins/Analytics/HordeServer.Analytics/Telemetry/Sinks/EpicTelemetrySink.cs b/Engine/Source/Programs/Horde/Plugins/Analytics/HordeServer.Analytics/Telemetry/Sinks/EpicTelemetrySink.cs
--- a/Engine/Source/Programs/Horde/Plugins/Analytics/HordeServer.Analytics/Telemetry/Sinks/EpicTelemetrySink.cs
+++ b/Engine/Source/Programs/Horde/Plugins/Analytics/HordeServer.Analytics/Telemetry/Sinks/EpicTelemetrySink.cs
@@ -60,7 +60,7 @@ namespace HordeServer.Telemetry.Sinks
 		static readonly Utf8String s_packetPrefix = new Utf8String("{\"Events\":[");
 		static readonly Utf8String s_packetSuffix = new Utf8String("]}");
 
-		readonly object _lockObject = new object();
+		readonly Lock _lockObject = new Lock();
 
 		readonly HttpClient _httpClient;
 		readonly Uri? _baseUrl;
diff --git a/Engine/Source/Programs/Horde/Plugins/Build/HordeServer.Build.Tests/HordeServer.Build.Tests.csproj b/Engine/Source/Programs/Horde/Plugins/Build/HordeServer.Build.Tests/HordeServer.Build.Tests.csproj
--- a/Engine/Source/Programs/Horde/Plugins/Build/HordeServer.Build.Tests/HordeServer.Build.Tests.csproj
+++ b/Engine/Source/Programs/Horde/Plugins/Build/HordeServer.Build.Tests/HordeServer.Build.Tests.csproj
@@ -6,9 +6,8 @@
   </PropertyGroup>
 
   <ItemGroup>
-    <PackageReference Include="Microsoft.AspNetCore.Mvc.Testing" Version="8.0.24" />
+    <PackageReference Include="Microsoft.AspNetCore.Mvc.Testing" Version="10.0.3" />
     <PackageReference Include="Moq" Version="4.20.72" />
-    <PackageReference Include="Microsoft.Extensions.Caching.Memory" Version="8.0.1" />
   </ItemGroup>
 
   <ItemGroup>
diff --git a/Engine/Source/Programs/Horde/Plugins/Build/HordeServer.Build.Tests/Server/RpcServiceTest.cs b/Engine/Source/Programs/Horde/Plugins/Build/HordeServer.Build.Tests/Server/RpcServiceTest.cs
--- a/Engine/Source/Programs/Horde/Plugins/Build/HordeServer.Build.Tests/Server/RpcServiceTest.cs
+++ b/Engine/Source/Programs/Horde/Plugins/Build/HordeServer.Build.Tests/Server/RpcServiceTest.cs
@@ -295,7 +295,7 @@ namespace HordeServer.Tests.Server
 				using (ZipArchive archive = new (stream, ZipArchiveMode.Create, leaveOpen: true))
 				{
 					ZipArchiveEntry entry = archive.CreateEntry(FileName);
-					await using (Stream entryStream = entry.Open())
+					await using (Stream entryStream = await entry.OpenAsync(cancellationToken))
 					{
 						await entryStream.WriteAsync(fileData, cancellationToken);
 					}
diff --git a/Engine/Source/Programs/Horde/Plugins/Build/HordeServer.Build/Jobs/JobsController.cs b/Engine/Source/Programs/Horde/Plugins/Build/HordeServer.Build/Jobs/JobsController.cs
--- a/Engine/Source/Programs/Horde/Plugins/Build/HordeServer.Build/Jobs/JobsController.cs
+++ b/Engine/Source/Programs/Horde/Plugins/Build/HordeServer.Build/Jobs/JobsController.cs
@@ -1174,14 +1174,15 @@ namespace HordeServer.Jobs
 
 			IReadOnlyList<IJob> jobs = await _jobService.FindJobsByStreamWithTemplatesAsync(new StreamId(streamId), templateRefIds, count: count, consistentRead: false, cancellationToken: cancellationToken);
 
-			Dictionary<string, GetJobTimingResponse> jobTimings = await jobs.ToAsyncEnumerable()
+			Dictionary<string, GetJobTimingResponse> jobTimings = [];
+			await foreach (IJob job in jobs.ToAsyncEnumerable()
 				.Where(job => _buildConfig.Value.Authorize(job, JobAclAction.ViewJob, User))
-				.ToDictionaryAwaitAsync(x => ValueTask.FromResult(x.Id.ToString()), async job =>
-				{
-					IJobTiming jobTiming = await _jobService.GetJobTimingAsync(job, _logger, cancellationToken);
-					IGraph graph = await _jobService.GetGraphAsync(job);
-					return await CreateJobTimingResponseAsync(job, graph, jobTiming, _logger, true);
-				}, cancellationToken);
+				.WithCancellation(cancellationToken))
+			{
+				IJobTiming jobTiming = await _jobService.GetJobTimingAsync(job, _logger, cancellationToken);
+				IGraph graph = await _jobService.GetGraphAsync(job, cancellationToken);
+				jobTimings[job.Id.ToString()] = await CreateJobTimingResponseAsync(job, graph, jobTiming, _logger, true, cancellationToken);
+			}
 
 			return PropertyFilter.Apply(new FindJobTimingsResponse(jobTimings), filter);
 		}
diff --git a/Engine/Source/Programs/Horde/Plugins/Build/HordeServer.Build/Jobs/Schedules/ScheduleService.cs b/Engine/Source/Programs/Horde/Plugins/Build/HordeServer.Build/Jobs/Schedules/ScheduleService.cs
--- a/Engine/Source/Programs/Horde/Plugins/Build/HordeServer.Build/Jobs/Schedules/ScheduleService.cs
+++ b/Engine/Source/Programs/Horde/Plugins/Build/HordeServer.Build/Jobs/Schedules/ScheduleService.cs
@@ -444,7 +444,7 @@ namespace HordeServer.Jobs.Schedules
 					commitId = await GetNextChangeForGateAsync(stream.Id, templateId, schedule.Gate, minCommitId, maxCommitId, scheduleLogger, cancellationToken);
 					commit = (commitId != null)? await commits.FindAsync(commitId, commitId, 1, null, cancellationToken).FirstOrDefaultAsync(cancellationToken) : null; // May be a change in a different stream
 				}
-				else if (await commitEnumerator.MoveNextAsync(cancellationToken))
+				else if (await commitEnumerator.MoveNextAsync())
 				{
 					commit = commitEnumerator.Current;
 					commitId = commit.Id;
diff --git a/Engine/Source/Programs/Horde/Plugins/Compute/HordeServer.Compute/Agents/AgentScheduler.cs b/Engine/Source/Programs/Horde/Plugins/Compute/HordeServer.Compute/Agents/AgentScheduler.cs
--- a/Engine/Source/Programs/Horde/Plugins/Compute/HordeServer.Compute/Agents/AgentScheduler.cs
+++ b/Engine/Source/Programs/Horde/Plugins/Compute/HordeServer.Compute/Agents/AgentScheduler.cs
@@ -479,7 +479,7 @@ namespace HordeServer.Agents
 			HashSet<SessionId> sessionStates = await server
 				.KeysAsync(database: database.Database, pattern: SessionStateKey("*"))
 				.Select(key => FromSessionStateKey(key.ToString()))
-				.ToHashSetAsync(cancellationToken);
+				.ToHashSetAsync(cancellationToken: cancellationToken);
 			
 			// Sessions registered in the "compute:sessions:current" hash key
 			HashSet<SessionId> registeredSessions = (await database.HashGetAllAsync(Keys.Sessions.Current))
diff --git a/Engine/Source/Programs/Horde/Plugins/Compute/HordeServer.Compute/HordeServer.Compute.csproj b/Engine/Source/Programs/Horde/Plugins/Compute/HordeServer.Compute/HordeServer.Compute.csproj
--- a/Engine/Source/Programs/Horde/Plugins/Compute/HordeServer.Compute/HordeServer.Compute.csproj
+++ b/Engine/Source/Programs/Horde/Plugins/Compute/HordeServer.Compute/HordeServer.Compute.csproj
@@ -49,7 +49,6 @@
       <PrivateAssets>all</PrivateAssets>
       <IncludeAssets>runtime; build; native; contentfiles; analyzers; buildtransitive</IncludeAssets>
     </PackageReference>
-    <PackageReference Include="System.Linq.Async" Version="6.0.3" />
   </ItemGroup>
 
   <ItemGroup>
diff --git a/Engine/Source/Programs/Horde/Plugins/ExamplePlugin/HordeServer.ExamplePlugin.Tests/HordeServer.ExamplePlugin.Tests.csproj b/Engine/Source/Programs/Horde/Plugins/ExamplePlugin/HordeServer.ExamplePlugin.Tests/HordeServer.ExamplePlugin.Tests.csproj
--- a/Engine/Source/Programs/Horde/Plugins/ExamplePlugin/HordeServer.ExamplePlugin.Tests/HordeServer.ExamplePlugin.Tests.csproj
+++ b/Engine/Source/Programs/Horde/Plugins/ExamplePlugin/HordeServer.ExamplePlugin.Tests/HordeServer.ExamplePlugin.Tests.csproj
@@ -6,9 +6,8 @@
   </PropertyGroup>
 
   <ItemGroup>
-    <PackageReference Include="Microsoft.AspNetCore.Mvc.Testing" Version="8.0.24" />
+    <PackageReference Include="Microsoft.AspNetCore.Mvc.Testing" Version="10.0.3" />
     <PackageReference Include="Moq" Version="4.20.72" />
-    <PackageReference Include="Microsoft.Extensions.Caching.Memory" Version="8.0.1" />
   </ItemGroup>
 
   <ItemGroup>
diff --git a/Engine/Source/Programs/Horde/Plugins/Experimental/HordeServer.Experimental.Tests/HordeServer.Experimental.Tests.csproj b/Engine/Source/Programs/Horde/Plugins/Experimental/HordeServer.Experimental.Tests/HordeServer.Experimental.Tests.csproj
--- a/Engine/Source/Programs/Horde/Plugins/Experimental/HordeServer.Experimental.Tests/HordeServer.Experimental.Tests.csproj
+++ b/Engine/Source/Programs/Horde/Plugins/Experimental/HordeServer.Experimental.Tests/HordeServer.Experimental.Tests.csproj
@@ -6,9 +6,8 @@
   </PropertyGroup>
 
   <ItemGroup>
-    <PackageReference Include="Microsoft.AspNetCore.Mvc.Testing" Version="8.0.24" />
+    <PackageReference Include="Microsoft.AspNetCore.Mvc.Testing" Version="10.0.3" />
     <PackageReference Include="Moq" Version="4.20.72" />
-    <PackageReference Include="Microsoft.Extensions.Caching.Memory" Version="8.0.1" />
   </ItemGroup>
 
   <ItemGroup>
diff --git a/Engine/Source/Programs/Horde/Plugins/Tools/HordeServer.Tools.Tests/Tools/ToolTests.cs b/Engine/Source/Programs/Horde/Plugins/Tools/HordeServer.Tools.Tests/Tools/ToolTests.cs
--- a/Engine/Source/Programs/Horde/Plugins/Tools/HordeServer.Tools.Tests/Tools/ToolTests.cs
+++ b/Engine/Source/Programs/Horde/Plugins/Tools/HordeServer.Tools.Tests/Tools/ToolTests.cs
@@ -73,7 +73,7 @@ namespace HordeServer.Tests.Tools
 				using (ZipArchive archive = new ZipArchive(stream, ZipArchiveMode.Create))
 				{
 					ZipArchiveEntry entry = archive.CreateEntry(FileName);
-					using (Stream entryStream = entry.Open())
+					await using (Stream entryStream = await entry.OpenAsync())
 					{
 						await entryStream.WriteAsync(fileData);
 					}
@@ -131,7 +131,7 @@ namespace HordeServer.Tests.Tools
 				byte[] outputFileData;
 				using (MemoryStream stream = new MemoryStream())
 				{
-					using Stream entryStream = entry.Open();
+					await using Stream entryStream = await entry.OpenAsync();
 					await entryStream.CopyToAsync(stream);
 					outputFileData = stream.ToArray();
 				}
@@ -192,7 +192,7 @@ namespace HordeServer.Tests.Tools
 			using (ZipArchive archive = new(stream, ZipArchiveMode.Create))
 			{
 				ZipArchiveEntry entry = archive.CreateEntry(fileName);
-				await using (Stream entryStream = entry.Open())
+				await using (Stream entryStream = await entry.OpenAsync())
 				{
 					await entryStream.WriteAsync(Encoding.UTF8.GetBytes(fileData));
 				}
diff --git a/Engine/Source/Programs/Horde/Scripts/test.sh b/Engine/Source/Programs/Horde/Scripts/test.sh
--- a/Engine/Source/Programs/Horde/Scripts/test.sh
+++ b/Engine/Source/Programs/Horde/Scripts/test.sh
@@ -1,7 +1,7 @@
 #!/bin/bash
 # Run tests (run from Dockerfile)
 
-export UE_DOTNET_VERSION=net8.0
+export UE_DOTNET_VERSION=net10.0
 
 # Start Redis and MongoDB in the background for tests to use
 redis-server --save "" --appendonly no --daemonize yes --databases 3000 || exit 1
diff --git a/Engine/Source/Programs/Horde/Tools/ReleaseNotes/ReleaseNotes.csproj b/Engine/Source/Programs/Horde/Tools/ReleaseNotes/ReleaseNotes.csproj
--- a/Engine/Source/Programs/Horde/Tools/ReleaseNotes/ReleaseNotes.csproj
+++ b/Engine/Source/Programs/Horde/Tools/ReleaseNotes/ReleaseNotes.csproj
@@ -8,7 +8,6 @@
 
   <ItemGroup>
     <PackageReference Include="Microsoft.Extensions.Logging.Console" Version="8.0.1" />
-    <PackageReference Include="System.Text.Json" Version="8.0.6" />
   </ItemGroup>
 
   <ItemGroup>
diff --git a/Engine/Source/Programs/UnrealToolbox/UnrealToolbox.csproj b/Engine/Source/Programs/UnrealToolbox/UnrealToolbox.csproj
--- a/Engine/Source/Programs/UnrealToolbox/UnrealToolbox.csproj
+++ b/Engine/Source/Programs/UnrealToolbox/UnrealToolbox.csproj
@@ -10,7 +10,6 @@
     <Nullable>enable</Nullable>
     <ImplicitUsings>enable</ImplicitUsings>
     <NeutralLanguage>en-US</NeutralLanguage>
-    <Configurations>Debug;Release;Analyze</Configurations>
     <AssemblyName>UnrealToolbox</AssemblyName>
     <AvaloniaUseCompiledBindingsByDefault>true</AvaloniaUseCompiledBindingsByDefault>
     <NoWarn>$(NoWarn);CA1501</NoWarn>
@@ -19,11 +18,11 @@
   </PropertyGroup>
   
   <PropertyGroup Condition="'$(EnableForTPS)'=='true'">
-     <TargetFramework>net8.0-windows</TargetFramework>
+     <TargetFramework>net10.0-windows</TargetFramework>
    </PropertyGroup>
 
   <PropertyGroup Condition="'$(EnableForTPS)'!='true'">
-     <TargetFramework>net8.0-windows10.0.17763.0</TargetFramework>
+     <TargetFramework>net10.0-windows10.0.17763.0</TargetFramework>
      <WindowsSdkPackageVersion>10.0.17763.41</WindowsSdkPackageVersion>
    </PropertyGroup>
   
