Horde: Fix JobView label layout when no label columns defined

diff --git a/Engine/Source/Programs/Horde/HordeDashboard/src/components/JobView.tsx b/Engine/Source/Programs/Horde/HordeDashboard/src/components/JobView.tsx
--- a/Engine/Source/Programs/Horde/HordeDashboard/src/components/JobView.tsx
+++ b/Engine/Source/Programs/Horde/HordeDashboard/src/components/JobView.tsx
@@ -439,12 +439,12 @@ const JobList: React.FC<{ tab: string; filter: JobFilterSimple; aliasTemplates?:
             width: "100%"
          }}>
                {categories.map(cat => {
+                  const isOtherCategory = cat.category === "Other" || cat.name === "Other" || !cat.category;
                   const labels = filterLabels(item.job!, cat.category);
-                  if (!labels.length && !(cat.category === "Other" && item.job.defaultLabel)) return null;
 
                   // Add default label for "Other" category
                   const defaultLabel = item.job.defaultLabel;
-                  if (defaultLabel && (cat.category === "Other" || !cat.category)) {
+                  if (defaultLabel && isOtherCategory) {
                      defaultLabel.dashboardName = "Other";
                      if (!labels.find(l => l === defaultLabel)) {
                         labels.push(defaultLabel);
@@ -453,6 +453,8 @@ const JobList: React.FC<{ tab: string; filter: JobFilterSimple; aliasTemplates?:
 
                   if (!labels.length) return null;
 
+                  if (!labels.length) return null;
+
                   // Sort labels: failures first, then warnings, then by name
                   const outcomePriority: Record<string, number> = {
                      [LabelOutcome.Failure]: 0,
@@ -479,12 +481,19 @@ const JobList: React.FC<{ tab: string; filter: JobFilterSimple; aliasTemplates?:
                   let catState: LabelState = LabelState.Complete;
                   let catOutcome: LabelOutcome = LabelOutcome.Success;
                   for (const label of labels) {
-                     const aggregates = item.job.labels;
-                     const jlabel = aggregates?.find((l, idx) => l.dashboardCategory === label.dashboardCategory && l.dashboardName === label.dashboardName && item.job.labels![idx]?.state !== LabelState.Unspecified);
-                     if (!jlabel) { continue; }
-                     const idx = aggregates?.indexOf(jlabel)!;
-                     const ls = item.job.labels![idx]?.state;
-                     const lo = item.job.labels![idx]?.outcome;
+                     let ls: LabelState | undefined;
+                     let lo: LabelOutcome | undefined;
+                     if (label === item.job.defaultLabel) {
+                        ls = label.state;
+                        lo = label.outcome;
+                     } else {
+                        const aggregates = item.job.labels;
+                        const jlabel = aggregates?.find((l, idx) => l.dashboardCategory === label.dashboardCategory && l.dashboardName === label.dashboardName && item.job.labels![idx]?.state !== LabelState.Unspecified);
+                        if (!jlabel) { continue; }
+                        const idx = aggregates?.indexOf(jlabel)!;
+                        ls = item.job.labels![idx]?.state;
+                        lo = item.job.labels![idx]?.outcome;
+                     }
                      if (ls === LabelState.Running) catState = LabelState.Running;
                      if (lo === LabelOutcome.Failure) catOutcome = LabelOutcome.Failure;
                      else if (lo === LabelOutcome.Warnings && catOutcome !== LabelOutcome.Failure) catOutcome = LabelOutcome.Warnings;
diff --git a/Engine/Source/Programs/Horde/HordeDashboard/src/components/JobViewAll.tsx b/Engine/Source/Programs/Horde/HordeDashboard/src/components/JobViewAll.tsx
--- a/Engine/Source/Programs/Horde/HordeDashboard/src/components/JobViewAll.tsx
+++ b/Engine/Source/Programs/Horde/HordeDashboard/src/components/JobViewAll.tsx
@@ -394,21 +394,7 @@ const JobViewAllInner: React.FC<{ filter: JobFilterSimple }> = observer(({ filte
 
       const target = `label_${item.job.id}_${label.dashboardName}_${label.dashboardCategory}`.replace(/ /g, "");
 
-      return <Stack>
-         <div id={target}>
-            <LabelHoverCard
-               jobId={item.job.id}
-               labelIndex={labelIdx}
-               streamId={streamId!}
-               templateId={item.job.templateId!}
-               templateName={item.job.name}
-               dashboardName={label.dashboardName}
-               dashboardCategory={label.dashboardCategory}
-               state={state}
-               outcome={outcome}
-               change={item.job.change}
-            >
-               <Link to={url}><Stack className={hordeClasses.badgeNoIcon}>
+      const badge = <Link to={url}><Stack className={hordeClasses.badgeNoIcon}>
                   <DefaultButton key={label.dashboardName} style={{ backgroundColor: color.primaryColor, fontSize: 9 }} text={label.dashboardName}>
                      {!!color.secondaryColor && <div style={{
                         borderLeft: "10px solid transparent",
@@ -423,8 +409,24 @@ const JobViewAllInner: React.FC<{ filter: JobFilterSimple }> = observer(({ filte
                      }} />}
                   </DefaultButton>
                </Stack>
-               </Link>
-            </LabelHoverCard>
+               </Link>;
+
+      return <Stack>
+         <div id={target}>
+            {labelIdx >= 0 ? <LabelHoverCard
+               jobId={item.job.id}
+               labelIndex={labelIdx}
+               streamId={streamId!}
+               templateId={item.job.templateId!}
+               templateName={item.job.name}
+               dashboardName={label.dashboardName}
+               dashboardCategory={label.dashboardCategory}
+               state={state}
+               outcome={outcome}
+               change={item.job.change}
+            >
+               {badge}
+            </LabelHoverCard> : badge}
          </div>
       </Stack>;
    };
diff --git a/Engine/Source/Programs/Horde/HordeDashboard/src/components/JobViewIncremental.tsx b/Engine/Source/Programs/Horde/HordeDashboard/src/components/JobViewIncremental.tsx
--- a/Engine/Source/Programs/Horde/HordeDashboard/src/components/JobViewIncremental.tsx
+++ b/Engine/Source/Programs/Horde/HordeDashboard/src/components/JobViewIncremental.tsx
@@ -411,21 +411,7 @@ const JobList: React.FC<{ tab: string; filter: JobFilterSimple; aliasTemplates?:
 
       const target = `label_${item.job.id}_${label.dashboardName}_${label.dashboardCategory}`.replace(/[^A-Za-z0-9]/g, "");
 
-      return <Stack>
-         <div id={target} className="horde-no-darktheme">
-            <LabelHoverCard
-               jobId={item.job.id}
-               labelIndex={labelIdx}
-               streamId={streamId!}
-               templateId={item.job.templateId!}
-               templateName={item.job.name}
-               dashboardName={label.dashboardName}
-               dashboardCategory={label.dashboardCategory}
-               state={state}
-               outcome={outcome}
-               change={item.job.change}
-            >
-               <Link to={url}><Stack className={hordeClasses.badgeNoIcon}><DefaultButton key={label.dashboardName} style={{ backgroundColor: color.primaryColor }} text={label.dashboardName}
+      const badge = <Link to={url}><Stack className={hordeClasses.badgeNoIcon}><DefaultButton key={label.dashboardName} style={{ backgroundColor: color.primaryColor }} text={label.dashboardName}
                   onMouseOver={(ev) => ev.stopPropagation() }
                   onMouseMove={(ev) => ev.stopPropagation()}>
 
@@ -441,8 +427,24 @@ const JobList: React.FC<{ tab: string; filter: JobFilterSimple; aliasTemplates?:
                      zIndex: 1
                   }} />}
 
-               </DefaultButton></Stack></Link>
-            </LabelHoverCard>
+               </DefaultButton></Stack></Link>;
+
+      return <Stack>
+         <div id={target} className="horde-no-darktheme">
+            {labelIdx >= 0 ? <LabelHoverCard
+               jobId={item.job.id}
+               labelIndex={labelIdx}
+               streamId={streamId!}
+               templateId={item.job.templateId!}
+               templateName={item.job.name}
+               dashboardName={label.dashboardName}
+               dashboardCategory={label.dashboardCategory}
+               state={state}
+               outcome={outcome}
+               change={item.job.change}
+            >
+               {badge}
+            </LabelHoverCard> : badge}
          </div>
       </Stack>;
 
diff --git a/Engine/Source/Programs/Horde/HordeDashboard/src/components/labels/LabelBadge.tsx b/Engine/Source/Programs/Horde/HordeDashboard/src/components/labels/LabelBadge.tsx
--- a/Engine/Source/Programs/Horde/HordeDashboard/src/components/labels/LabelBadge.tsx
+++ b/Engine/Source/Programs/Horde/HordeDashboard/src/components/labels/LabelBadge.tsx
@@ -48,6 +48,11 @@ export const LabelBadge: React.FC<LabelBadgeProps> = ({ label, jobId, labelIndex
         </div>
     );
 
+    // Only show hover card for actual defined labels (not the procedural defaultLabel "Other")
+    if (labelIndex < 0) {
+        return badge;
+    }
+
     return (
         <LabelHoverCard
             jobId={jobId}
