[Horde Dashboard] Ensure unnamed pools aren't displayed in pool view filter results

diff --git a/Engine/Source/Programs/Horde/HordeDashboard/src/components/PoolsView.tsx b/Engine/Source/Programs/Horde/HordeDashboard/src/components/PoolsView.tsx
--- a/Engine/Source/Programs/Horde/HordeDashboard/src/components/PoolsView.tsx
+++ b/Engine/Source/Programs/Horde/HordeDashboard/src/components/PoolsView.tsx
@@ -422,9 +422,12 @@ const PoolPicker: React.FC = observer(() => {
 
    handler.subscribe();
 
-   const poolTags: ITag[] = handler.pools.sort((a, b) => a.name.localeCompare(b.name)).map(p => {
-      return { key: p.id, name: p.name }
-   });
+   const poolTags: ITag[] = handler.pools
+      .filter(pool => pool.name !== "")
+      .sort((a, b) => a.name.localeCompare(b.name))
+      .map(p => {
+         return { key: p.id, name: p.name };
+      });
 
    let selectedItems: ITag[] = [];
    if (poolId) {
diff --git a/Engine/Source/Programs/Horde/HordeDashboard/src/components/agents/AgentStore.ts b/Engine/Source/Programs/Horde/HordeDashboard/src/components/agents/AgentStore.ts
--- a/Engine/Source/Programs/Horde/HordeDashboard/src/components/agents/AgentStore.ts
+++ b/Engine/Source/Programs/Horde/HordeDashboard/src/components/agents/AgentStore.ts
@@ -152,6 +152,10 @@ class AgentStore {
             });
         });
     }
+
+    filterValidPools(): PoolData[] {
+        return this.pools.filter(pool => pool.name !== "");
+    }
 }
 
 export const agentStore = new AgentStore();
\ No newline at end of file
diff --git a/Engine/Source/Programs/Horde/HordeDashboard/src/components/agents/AgentView.tsx b/Engine/Source/Programs/Horde/HordeDashboard/src/components/agents/AgentView.tsx
--- a/Engine/Source/Programs/Horde/HordeDashboard/src/components/agents/AgentView.tsx
+++ b/Engine/Source/Programs/Horde/HordeDashboard/src/components/agents/AgentView.tsx
@@ -945,7 +945,7 @@ class EditPoolsModalState {
    @action
    setOpen(isEditorOpen?: boolean, selectedItem?: PoolData) {
       this._getNumAgentsAssigned();
-      this.modifiedPools = [...agentStore.pools.map(pool => { return this._poolToPoolItemCopy(pool); })];
+      this.modifiedPools = [...agentStore.filterValidPools().map(pool => { return this._poolToPoolItemCopy(pool); })];
       this.newIdSuffix = 0;
       this.isOpen = true;
       this.isDirectEdit = isEditorOpen ?? false;
@@ -1175,7 +1175,7 @@ class SelectPoolsModalState {
       // if we're closing don't reset
       if (isOpen) {
          // first get all nondeleted pools
-         this.availablePools = agentStore.pools.map(pool => { return { key: pool.id, name: pool.name, pool: pool, agentsAssigned: [] }; });
+         this.availablePools = agentStore.filterValidPools().map(pool => { return { key: pool.id, name: pool.name, pool: pool, agentsAssigned: [] }; });
          this.agentChangeState = [];
          this.selectedPools = [];
          this.agentsSelectedTooltip = "";
