Horde: Pass job ID to label hover card ChangeButton for context menu

ChangeButton in LabelHoverCard was missing job ID, causing Open Job
Details and Copy Job links to point to /job/undefined. Pass jobId for
current build and prev.jobId for previous build CL buttons.

diff --git a/Engine/Source/Programs/Horde/HordeDashboard/src/components/labels/LabelHoverCard.tsx b/Engine/Source/Programs/Horde/HordeDashboard/src/components/labels/LabelHoverCard.tsx
--- a/Engine/Source/Programs/Horde/HordeDashboard/src/components/labels/LabelHoverCard.tsx
+++ b/Engine/Source/Programs/Horde/HordeDashboard/src/components/labels/LabelHoverCard.tsx
@@ -7,7 +7,7 @@ import backend from '../../backend';
 import { GetLabelDetailsResponse, GetLabelIssueResponse, GetLabelStepResponse, GetLogEventResponse, IssueSeverity, JobStepOutcome, JobStepState, LabelOutcome, LabelState } from '../../backend/Api';
 import dashboard, { StatusColor } from '../../backend/Dashboard';
 import { getLabelColor } from '../../styles/colors';
-import { ChangeButton, ChangeContextMenu } from '../ChangeButton';
+import { ChangeButton } from '../ChangeButton';
 import { IssueModalV2 } from '../IssueViewV2';
 
 // Client-side cache for label details responses (avoids re-fetching on re-hover)
@@ -118,7 +118,7 @@ const LabelHoverContent: React.FC<LabelHoverCardProps> = ({ jobId, labelIndex, s
                     <Text style={{ fontFamily: "Horde Open Sans Bold", fontSize: 14 }}>
                         {dashboardCategory ? `${dashboardCategory} / ` : ''}{dashboardName ?? details?.dashboardName ?? ''}
                     </Text>
-                    <ChangeButton job={{ streamId, change }} />
+                    <ChangeButton job={{ id: jobId, streamId, change }} />
                 </Stack>
                 {details?.previousBuild && <PreviousBuildSection details={details} streamId={streamId} />}
             </Stack>
@@ -183,7 +183,7 @@ const PreviousBuildSection: React.FC<{ details: GetLabelDetailsResponse; streamI
                     </Text>
                 </div>
             </Link>
-            <ChangeButton job={{ streamId, change: prev.change }} />
+            <ChangeButton job={{ id: prev.jobId, streamId, change: prev.change }} />
             <Text variant="small" style={{ opacity: 0.6 }}>Previous Job</Text>
         </Stack>
     );
