Horde: Updating schema documentation

#rnx

diff --git a/Engine/Source/Programs/Horde/Docs/Config/Schema/AclActions.md b/Engine/Source/Programs/Horde/Docs/Config/Schema/AclActions.md
--- a/Engine/Source/Programs/Horde/Docs/Config/Schema/AclActions.md
+++ b/Engine/Source/Programs/Horde/Docs/Config/Schema/AclActions.md
@@ -37,31 +37,47 @@
 | `DeleteAccount` |  |
 | `ViewAccount` |  |
 
-## Notices
+## Artifacts
 
 | Name | Description |
 | ---- | ----------- |
-| `CreateNotice` |  |
-| `UpdateNotice` |  |
-| `DeleteNotice` |  |
+| `ReadArtifact` | Permission to read from an artifact |
+| `WriteArtifact` | Permission to write to an artifact |
+| `DeleteArtifact` | Permission to delete to an artifact |
+| `UploadArtifact` | Ability to create an artifact. Typically just for debugging; agents have this access for a particular session. |
+| `DownloadArtifact` | Ability to download an artifact |
 
-## ServiceAccounts
+## Bisect
 
 | Name | Description |
 | ---- | ----------- |
-| `CreateAccount` |  |
-| `UpdateAccount` |  |
-| `DeleteAccount` |  |
-| `ViewAccount` |  |
+| `CreateBisectTask` | Ability to start new bisect tasks |
+| `UpdateBisectTask` | Ability to update a bisect task |
+| `ViewBisectTask` | Ability to view a bisect task |
 
-## Accounts
+## BuildHealth
 
 | Name | Description |
 | ---- | ----------- |
-| `CreateAccount` |  |
-| `UpdateAccount` |  |
-| `DeleteAccount` |  |
-| `ViewAccount` |  |
+| `AddUpdateDeleteBuildHealthFilter` | Ability to add, update, and delete build health filters. |
+
+## Devices
+
+| Name | Description |
+| ---- | ----------- |
+| `DeviceRead` | Ability to read devices |
+| `DeviceWrite` | Ability to write devices |
+
+## Jobs
+
+| Name | Description |
+| ---- | ----------- |
+| `CreateJob` | Ability to start new jobs |
+| `UpdateJob` | Rename a job, modify its priority, etc... |
+| `DeleteJob` | Delete a job properties |
+| `ExecuteJob` | Allows updating a job metadata (name, changelist number, step properties, new groups, job states, etc...). Typically granted to agents. Not user facing. |
+| `RetryJobStep` | Ability to retry a failed job step |
+| `ViewJob` | Ability to view a job |
 
 ## Notices
 
@@ -71,31 +87,27 @@
 | `UpdateNotice` |  |
 | `DeleteNotice` |  |
 
-## ServiceAccounts
+## Notifications
 
 | Name | Description |
 | ---- | ----------- |
-| `CreateAccount` |  |
-| `UpdateAccount` |  |
-| `DeleteAccount` |  |
-| `ViewAccount` |  |
+| `CreateSubscription` | Ability to subscribe to notifications |
 
-## Accounts
+## Projects
 
 | Name | Description |
 | ---- | ----------- |
-| `CreateAccount` |  |
-| `UpdateAccount` |  |
-| `DeleteAccount` |  |
-| `ViewAccount` |  |
+| `CreateProject` | Allows the creation of new projects |
+| `DeleteProject` | Allows deletion of projects. |
+| `UpdateProject` | Modify attributes of a project (name, categories, etc...) |
+| `ViewProject` | View information about a project |
 
-## Notices
+## Replicators
 
 | Name | Description |
 | ---- | ----------- |
-| `CreateNotice` |  |
-| `UpdateNotice` |  |
-| `DeleteNotice` |  |
+| `UpdateReplicator` | Allows deletion of projects. |
+| `ViewReplicator` | Allows the creation of new projects |
 
 ## ServiceAccounts
 
@@ -106,6 +118,24 @@
 | `DeleteAccount` |  |
 | `ViewAccount` |  |
 
+## Streams
+
+| Name | Description |
+| ---- | ----------- |
+| `CreateStream` | Allows the creation of new streams within a project |
+| `UpdateStream` | Allows updating a stream (agent types, templates, schedules) |
+| `DeleteStream` | Allows deleting a stream |
+| `ViewStream` | Ability to view a stream |
+| `ViewChanges` | View changes submitted to a stream. NOTE: this returns responses from the server's Perforce account, which may be a priviledged user. |
+| `ViewTemplate` | View template associated with a stream |
+
+## TestData
+
+| Name | Description |
+| ---- | ----------- |
+| `TestHealthRead` | Ability to read test health |
+| `TestHealthWrite` | Ability to write test health |
+
 ## Accounts
 
 | Name | Description |
@@ -115,37 +145,85 @@
 | `DeleteAccount` |  |
 | `ViewAccount` |  |
 
-## Notices
+## Agents
 
 | Name | Description |
 | ---- | ----------- |
-| `CreateNotice` |  |
-| `UpdateNotice` |  |
-| `DeleteNotice` |  |
+| `CreateAgent` | Ability to create an agent. This may be done explicitly, or granted to agents to allow them to self-register. |
+| `CreateWorkstationAgent` | Ability to create a workstation agent. Granted to users with hardware to offer. |
+| `UpdateAgent` | Update an agent's name, pools, etc... |
+| `DeleteAgent` | Soft-delete an agent |
+| `ViewAgent` | View an agent |
+| `ListAgents` | List the available agents |
 
-## Secrets
+## Artifacts
 
 | Name | Description |
 | ---- | ----------- |
-| `ViewSecret` | View a credential |
+| `ReadArtifact` |  |
+| `WriteArtifact` |  |
+| `DeleteArtifact` |  |
+| `UploadArtifact` |  |
+| `DownloadArtifact` |  |
 
-## ServiceAccounts
+## Bisect
 
 | Name | Description |
 | ---- | ----------- |
-| `CreateAccount` |  |
-| `UpdateAccount` |  |
-| `DeleteAccount` |  |
-| `ViewAccount` |  |
+| `CreateBisectTask` |  |
+| `UpdateBisectTask` |  |
+| `ViewBisectTask` |  |
 
-## Accounts
+## BuildHealth
 
 | Name | Description |
 | ---- | ----------- |
-| `CreateAccount` |  |
-| `UpdateAccount` |  |
-| `DeleteAccount` |  |
-| `ViewAccount` |  |
+| `AddUpdateDeleteBuildHealthFilter` |  |
+
+## Compute
+
+| Name | Description |
+| ---- | ----------- |
+| `AddComputeTasks` | User can add tasks to the compute cluster |
+| `GetComputeTasks` | User can get and list tasks from the compute cluster |
+| `UbaCacheRead` | User can read from Unreal Build Accelerator cache |
+| `UbaCacheWrite` | User can write to Unreal Build Accelerator cache |
+
+## Devices
+
+| Name | Description |
+| ---- | ----------- |
+| `DeviceRead` |  |
+| `DeviceWrite` |  |
+
+## Jobs
+
+| Name | Description |
+| ---- | ----------- |
+| `CreateJob` |  |
+| `UpdateJob` |  |
+| `DeleteJob` |  |
+| `ExecuteJob` |  |
+| `RetryJobStep` |  |
+| `ViewJob` |  |
+
+## Leases
+
+| Name | Description |
+| ---- | ----------- |
+| `ViewLeases` | View all the leases that an agent has worked on |
+| `ViewLeaseTasks` | View the task data for a lease |
+
+## Logs
+
+| Name | Description |
+| ---- | ----------- |
+| `CreateLog` | Ability to create a log. Implicitly granted to agents. |
+| `UpdateLog` | Ability to update log metadata |
+| `ViewLog` | Ability to view a log contents |
+| `WriteLogData` | Ability to write log data |
+| `CreateEvent` | Ability to create events |
+| `ViewEvent` | Ability to view events |
 
 ## Notices
 
@@ -155,11 +233,37 @@
 | `UpdateNotice` |  |
 | `DeleteNotice` |  |
 
-## Secrets
+## Notifications
 
 | Name | Description |
 | ---- | ----------- |
-| `ViewSecret` |  |
+| `CreateSubscription` |  |
 
+## Pools
+
+| Name | Description |
+| ---- | ----------- |
+| `CreatePool` | Create a global pool of agents |
+| `UpdatePool` | Modify an agent pool |
+| `DeletePool` | Delete an agent pool |
+| `ViewPool` | Ability to view a pool |
+| `ListPools` | View all the available agent pools |
+
+## Projects
+
+| Name | Description |
+| ---- | ----------- |
+| `CreateProject` |  |
+| `DeleteProject` |  |
+| `UpdateProject` |  |
+| `ViewProject` |  |
+
+## Replicators
+
+| Name | Description |
+| ---- | ----------- |
+| `UpdateReplicator` |  |
+| `ViewReplicator` |  |
+
 ## ServiceAccounts
 
@@ -170,11 +274,38 @@
 | `DeleteAccount` |  |
 | `ViewAccount` |  |
 
-## Symbols
+## Sessions
 
 | Name | Description |
 | ---- | ----------- |
-| `ReadSymbols` | Ability to download symbols |
+| `CreateSession` | Granted to agents to call CreateSession, which returns a bearer token identifying themselves valid to call UpdateSesssion via gRPC. |
+| `ViewSession` | Allows viewing information about an agent session |
 
+## Software
+
+| Name | Description |
+| ---- | ----------- |
+| `UploadSoftware` | Ability to upload new versions of the agent software |
+| `DownloadSoftware` | Ability to download the agent software |
+| `DeleteSoftware` | Ability to delete agent software |
+
+## Streams
+
+| Name | Description |
+| ---- | ----------- |
+| `CreateStream` |  |
+| `UpdateStream` |  |
+| `DeleteStream` |  |
+| `ViewStream` |  |
+| `ViewChanges` |  |
+| `ViewTemplate` |  |
+
+## TestData
+
+| Name | Description |
+| ---- | ----------- |
+| `TestHealthRead` |  |
+| `TestHealthWrite` |  |
+
 ## Accounts
 
@@ -185,6 +316,86 @@
 | `DeleteAccount` |  |
 | `ViewAccount` |  |
 
+## Agents
+
+| Name | Description |
+| ---- | ----------- |
+| `CreateAgent` |  |
+| `CreateWorkstationAgent` |  |
+| `UpdateAgent` |  |
+| `DeleteAgent` |  |
+| `ViewAgent` |  |
+| `ListAgents` |  |
+
+## Artifacts
+
+| Name | Description |
+| ---- | ----------- |
+| `ReadArtifact` |  |
+| `WriteArtifact` |  |
+| `DeleteArtifact` |  |
+| `UploadArtifact` |  |
+| `DownloadArtifact` |  |
+
+## Bisect
+
+| Name | Description |
+| ---- | ----------- |
+| `CreateBisectTask` |  |
+| `UpdateBisectTask` |  |
+| `ViewBisectTask` |  |
+
+## BuildHealth
+
+| Name | Description |
+| ---- | ----------- |
+| `AddUpdateDeleteBuildHealthFilter` |  |
+
+## Compute
+
+| Name | Description |
+| ---- | ----------- |
+| `AddComputeTasks` |  |
+| `GetComputeTasks` |  |
+| `UbaCacheRead` |  |
+| `UbaCacheWrite` |  |
+
+## Devices
+
+| Name | Description |
+| ---- | ----------- |
+| `DeviceRead` |  |
+| `DeviceWrite` |  |
+
+## Jobs
+
+| Name | Description |
+| ---- | ----------- |
+| `CreateJob` |  |
+| `UpdateJob` |  |
+| `DeleteJob` |  |
+| `ExecuteJob` |  |
+| `RetryJobStep` |  |
+| `ViewJob` |  |
+
+## Leases
+
+| Name | Description |
+| ---- | ----------- |
+| `ViewLeases` |  |
+| `ViewLeaseTasks` |  |
+
+## Logs
+
+| Name | Description |
+| ---- | ----------- |
+| `CreateLog` |  |
+| `UpdateLog` |  |
+| `ViewLog` |  |
+| `WriteLogData` |  |
+| `CreateEvent` |  |
+| `ViewEvent` |  |
+
 ## Notices
 
 | Name | Description |
@@ -193,13 +404,606 @@
 | `UpdateNotice` |  |
 | `DeleteNotice` |  |
 
+## Notifications
+
+| Name | Description |
+| ---- | ----------- |
+| `CreateSubscription` |  |
+
+## Pools
+
+| Name | Description |
+| ---- | ----------- |
+| `CreatePool` |  |
+| `UpdatePool` |  |
+| `DeletePool` |  |
+| `ViewPool` |  |
+| `ListPools` |  |
+
+## Projects
+
+| Name | Description |
+| ---- | ----------- |
+| `CreateProject` |  |
+| `DeleteProject` |  |
+| `UpdateProject` |  |
+| `ViewProject` |  |
+
+## Replicators
+
+| Name | Description |
+| ---- | ----------- |
+| `UpdateReplicator` |  |
+| `ViewReplicator` |  |
+
+## ServiceAccounts
+
+| Name | Description |
+| ---- | ----------- |
+| `CreateAccount` |  |
+| `UpdateAccount` |  |
+| `DeleteAccount` |  |
+| `ViewAccount` |  |
+
+## Sessions
+
+| Name | Description |
+| ---- | ----------- |
+| `CreateSession` |  |
+| `ViewSession` |  |
+
+## Software
+
+| Name | Description |
+| ---- | ----------- |
+| `UploadSoftware` |  |
+| `DownloadSoftware` |  |
+| `DeleteSoftware` |  |
+
+## Streams
+
+| Name | Description |
+| ---- | ----------- |
+| `CreateStream` |  |
+| `UpdateStream` |  |
+| `DeleteStream` |  |
+| `ViewStream` |  |
+| `ViewChanges` |  |
+| `ViewTemplate` |  |
+
+## TestData
+
+| Name | Description |
+| ---- | ----------- |
+| `TestHealthRead` |  |
+| `TestHealthWrite` |  |
+
+## Accounts
+
+| Name | Description |
+| ---- | ----------- |
+| `CreateAccount` |  |
+| `UpdateAccount` |  |
+| `DeleteAccount` |  |
+| `ViewAccount` |  |
+
+## Agents
+
+| Name | Description |
+| ---- | ----------- |
+| `CreateAgent` |  |
+| `CreateWorkstationAgent` |  |
+| `UpdateAgent` |  |
+| `DeleteAgent` |  |
+| `ViewAgent` |  |
+| `ListAgents` |  |
+
+## Artifacts
+
+| Name | Description |
+| ---- | ----------- |
+| `ReadArtifact` |  |
+| `WriteArtifact` |  |
+| `DeleteArtifact` |  |
+| `UploadArtifact` |  |
+| `DownloadArtifact` |  |
+
+## Bisect
+
+| Name | Description |
+| ---- | ----------- |
+| `CreateBisectTask` |  |
+| `UpdateBisectTask` |  |
+| `ViewBisectTask` |  |
+
+## BuildHealth
+
+| Name | Description |
+| ---- | ----------- |
+| `AddUpdateDeleteBuildHealthFilter` |  |
+
+## Compute
+
+| Name | Description |
+| ---- | ----------- |
+| `AddComputeTasks` |  |
+| `GetComputeTasks` |  |
+| `UbaCacheRead` |  |
+| `UbaCacheWrite` |  |
+
+## Devices
+
+| Name | Description |
+| ---- | ----------- |
+| `DeviceRead` |  |
+| `DeviceWrite` |  |
+
+## Jobs
+
+| Name | Description |
+| ---- | ----------- |
+| `CreateJob` |  |
+| `UpdateJob` |  |
+| `DeleteJob` |  |
+| `ExecuteJob` |  |
+| `RetryJobStep` |  |
+| `ViewJob` |  |
+
+## Leases
+
+| Name | Description |
+| ---- | ----------- |
+| `ViewLeases` |  |
+| `ViewLeaseTasks` |  |
+
+## Logs
+
+| Name | Description |
+| ---- | ----------- |
+| `CreateLog` |  |
+| `UpdateLog` |  |
+| `ViewLog` |  |
+| `WriteLogData` |  |
+| `CreateEvent` |  |
+| `ViewEvent` |  |
+
+## Notices
+
+| Name | Description |
+| ---- | ----------- |
+| `CreateNotice` |  |
+| `UpdateNotice` |  |
+| `DeleteNotice` |  |
+
+## Notifications
+
+| Name | Description |
+| ---- | ----------- |
+| `CreateSubscription` |  |
+
+## Pools
+
+| Name | Description |
+| ---- | ----------- |
+| `CreatePool` |  |
+| `UpdatePool` |  |
+| `DeletePool` |  |
+| `ViewPool` |  |
+| `ListPools` |  |
+
+## Projects
+
+| Name | Description |
+| ---- | ----------- |
+| `CreateProject` |  |
+| `DeleteProject` |  |
+| `UpdateProject` |  |
+| `ViewProject` |  |
+
+## Replicators
+
+| Name | Description |
+| ---- | ----------- |
+| `UpdateReplicator` |  |
+| `ViewReplicator` |  |
+
+## ServiceAccounts
+
+| Name | Description |
+| ---- | ----------- |
+| `CreateAccount` |  |
+| `UpdateAccount` |  |
+| `DeleteAccount` |  |
+| `ViewAccount` |  |
+
+## Sessions
+
+| Name | Description |
+| ---- | ----------- |
+| `CreateSession` |  |
+| `ViewSession` |  |
+
+## Software
+
+| Name | Description |
+| ---- | ----------- |
+| `UploadSoftware` |  |
+| `DownloadSoftware` |  |
+| `DeleteSoftware` |  |
+
+## Streams
+
+| Name | Description |
+| ---- | ----------- |
+| `CreateStream` |  |
+| `UpdateStream` |  |
+| `DeleteStream` |  |
+| `ViewStream` |  |
+| `ViewChanges` |  |
+| `ViewTemplate` |  |
+
+## TestData
+
+| Name | Description |
+| ---- | ----------- |
+| `TestHealthRead` |  |
+| `TestHealthWrite` |  |
+
+## Accounts
+
+| Name | Description |
+| ---- | ----------- |
+| `CreateAccount` |  |
+| `UpdateAccount` |  |
+| `DeleteAccount` |  |
+| `ViewAccount` |  |
+
+## Agents
+
+| Name | Description |
+| ---- | ----------- |
+| `CreateAgent` |  |
+| `CreateWorkstationAgent` |  |
+| `UpdateAgent` |  |
+| `DeleteAgent` |  |
+| `ViewAgent` |  |
+| `ListAgents` |  |
+
+## Artifacts
+
+| Name | Description |
+| ---- | ----------- |
+| `ReadArtifact` |  |
+| `WriteArtifact` |  |
+| `DeleteArtifact` |  |
+| `UploadArtifact` |  |
+| `DownloadArtifact` |  |
+
+## Bisect
+
+| Name | Description |
+| ---- | ----------- |
+| `CreateBisectTask` |  |
+| `UpdateBisectTask` |  |
+| `ViewBisectTask` |  |
+
+## BuildHealth
+
+| Name | Description |
+| ---- | ----------- |
+| `AddUpdateDeleteBuildHealthFilter` |  |
+
+## Compute
+
+| Name | Description |
+| ---- | ----------- |
+| `AddComputeTasks` |  |
+| `GetComputeTasks` |  |
+| `UbaCacheRead` |  |
+| `UbaCacheWrite` |  |
+
+## Devices
+
+| Name | Description |
+| ---- | ----------- |
+| `DeviceRead` |  |
+| `DeviceWrite` |  |
+
+## Jobs
+
+| Name | Description |
+| ---- | ----------- |
+| `CreateJob` |  |
+| `UpdateJob` |  |
+| `DeleteJob` |  |
+| `ExecuteJob` |  |
+| `RetryJobStep` |  |
+| `ViewJob` |  |
+
+## Leases
+
+| Name | Description |
+| ---- | ----------- |
+| `ViewLeases` |  |
+| `ViewLeaseTasks` |  |
+
+## Logs
+
+| Name | Description |
+| ---- | ----------- |
+| `CreateLog` |  |
+| `UpdateLog` |  |
+| `ViewLog` |  |
+| `WriteLogData` |  |
+| `CreateEvent` |  |
+| `ViewEvent` |  |
+
+## Notices
+
+| Name | Description |
+| ---- | ----------- |
+| `CreateNotice` |  |
+| `UpdateNotice` |  |
+| `DeleteNotice` |  |
+
+## Notifications
+
+| Name | Description |
+| ---- | ----------- |
+| `CreateSubscription` |  |
+
+## Pools
+
+| Name | Description |
+| ---- | ----------- |
+| `CreatePool` |  |
+| `UpdatePool` |  |
+| `DeletePool` |  |
+| `ViewPool` |  |
+| `ListPools` |  |
+
+## Projects
+
+| Name | Description |
+| ---- | ----------- |
+| `CreateProject` |  |
+| `DeleteProject` |  |
+| `UpdateProject` |  |
+| `ViewProject` |  |
+
+## Replicators
+
+| Name | Description |
+| ---- | ----------- |
+| `UpdateReplicator` |  |
+| `ViewReplicator` |  |
+
 ## Secrets
 
 | Name | Description |
 | ---- | ----------- |
+| `ViewSecret` | View a credential |
+
+## ServiceAccounts
+
+| Name | Description |
+| ---- | ----------- |
+| `CreateAccount` |  |
+| `UpdateAccount` |  |
+| `DeleteAccount` |  |
+| `ViewAccount` |  |
+
+## Sessions
+
+| Name | Description |
+| ---- | ----------- |
+| `CreateSession` |  |
+| `ViewSession` |  |
+
+## Software
+
+| Name | Description |
+| ---- | ----------- |
+| `UploadSoftware` |  |
+| `DownloadSoftware` |  |
+| `DeleteSoftware` |  |
+
+## Streams
+
+| Name | Description |
+| ---- | ----------- |
+| `CreateStream` |  |
+| `UpdateStream` |  |
+| `DeleteStream` |  |
+| `ViewStream` |  |
+| `ViewChanges` |  |
+| `ViewTemplate` |  |
+
+## TestData
+
+| Name | Description |
+| ---- | ----------- |
+| `TestHealthRead` |  |
+| `TestHealthWrite` |  |
+
+## Accounts
+
+| Name | Description |
+| ---- | ----------- |
+| `CreateAccount` |  |
+| `UpdateAccount` |  |
+| `DeleteAccount` |  |
+| `ViewAccount` |  |
+
+## Agents
+
+| Name | Description |
+| ---- | ----------- |
+| `CreateAgent` |  |
+| `CreateWorkstationAgent` |  |
+| `UpdateAgent` |  |
+| `DeleteAgent` |  |
+| `ViewAgent` |  |
+| `ListAgents` |  |
+
+## Artifacts
+
+| Name | Description |
+| ---- | ----------- |
+| `ReadArtifact` |  |
+| `WriteArtifact` |  |
+| `DeleteArtifact` |  |
+| `UploadArtifact` |  |
+| `DownloadArtifact` |  |
+
+## Bisect
+
+| Name | Description |
+| ---- | ----------- |
+| `CreateBisectTask` |  |
+| `UpdateBisectTask` |  |
+| `ViewBisectTask` |  |
+
+## BuildHealth
+
+| Name | Description |
+| ---- | ----------- |
+| `AddUpdateDeleteBuildHealthFilter` |  |
+
+## Compute
+
+| Name | Description |
+| ---- | ----------- |
+| `AddComputeTasks` |  |
+| `GetComputeTasks` |  |
+| `UbaCacheRead` |  |
+| `UbaCacheWrite` |  |
+
+## Devices
+
+| Name | Description |
+| ---- | ----------- |
+| `DeviceRead` |  |
+| `DeviceWrite` |  |
+
+## Jobs
+
+| Name | Description |
+| ---- | ----------- |
+| `CreateJob` |  |
+| `UpdateJob` |  |
+| `DeleteJob` |  |
+| `ExecuteJob` |  |
+| `RetryJobStep` |  |
+| `ViewJob` |  |
+
+## Leases
+
+| Name | Description |
+| ---- | ----------- |
+| `ViewLeases` |  |
+| `ViewLeaseTasks` |  |
+
+## Logs
+
+| Name | Description |
+| ---- | ----------- |
+| `CreateLog` |  |
+| `UpdateLog` |  |
+| `ViewLog` |  |
+| `WriteLogData` |  |
+| `CreateEvent` |  |
+| `ViewEvent` |  |
+
+## Notices
+
+| Name | Description |
+| ---- | ----------- |
+| `CreateNotice` |  |
+| `UpdateNotice` |  |
+| `DeleteNotice` |  |
+
+## Notifications
+
+| Name | Description |
+| ---- | ----------- |
+| `CreateSubscription` |  |
+
+## Pools
+
+| Name | Description |
+| ---- | ----------- |
+| `CreatePool` |  |
+| `UpdatePool` |  |
+| `DeletePool` |  |
+| `ViewPool` |  |
+| `ListPools` |  |
+
+## Projects
+
+| Name | Description |
+| ---- | ----------- |
+| `CreateProject` |  |
+| `DeleteProject` |  |
+| `UpdateProject` |  |
+| `ViewProject` |  |
+
+## Replicators
+
+| Name | Description |
+| ---- | ----------- |
+| `UpdateReplicator` |  |
+| `ViewReplicator` |  |
+
+## Secrets
+
+| Name | Description |
+| ---- | ----------- |
 | `ViewSecret` |  |
 
 ## ServiceAccounts
 
 | Name | Description |
 | ---- | ----------- |
+| `CreateAccount` |  |
+| `UpdateAccount` |  |
+| `DeleteAccount` |  |
+| `ViewAccount` |  |
+
+## Sessions
+
+| Name | Description |
+| ---- | ----------- |
+| `CreateSession` |  |
+| `ViewSession` |  |
+
+## Software
+
+| Name | Description |
+| ---- | ----------- |
+| `UploadSoftware` |  |
+| `DownloadSoftware` |  |
+| `DeleteSoftware` |  |
+
+## Streams
+
+| Name | Description |
+| ---- | ----------- |
+| `CreateStream` |  |
+| `UpdateStream` |  |
+| `DeleteStream` |  |
+| `ViewStream` |  |
+| `ViewChanges` |  |
+| `ViewTemplate` |  |
+
+## TestData
+
+| Name | Description |
+| ---- | ----------- |
+| `TestHealthRead` |  |
+| `TestHealthWrite` |  |
+
+## Accounts
+
+| Name | Description |
+| ---- | ----------- |
@@ -208,69 +1012,85 @@
 | `DeleteAccount` |  |
 | `ViewAccount` |  |
 
-## Symbols
+## Agents
 
 | Name | Description |
 | ---- | ----------- |
-| `ReadSymbols` |  |
+| `CreateAgent` |  |
+| `CreateWorkstationAgent` |  |
+| `UpdateAgent` |  |
+| `DeleteAgent` |  |
+| `ViewAgent` |  |
+| `ListAgents` |  |
 
-## Tools
-
-| Name | Description |
-| ---- | ----------- |
-| `DownloadTool` | Ability to download a tool |
-| `UploadTool` | Ability to upload new tool versions |
-
-## Accounts
-
-| Name | Description |
-| ---- | ----------- |
-| `CreateAccount` |  |
-| `UpdateAccount` |  |
-| `DeleteAccount` |  |
-| `ViewAccount` |  |
-
 ## Artifacts
 
 | Name | Description |
 | ---- | ----------- |
-| `ReadArtifact` | Permission to read from an artifact |
-| `WriteArtifact` | Permission to write to an artifact |
-| `DeleteArtifact` | Permission to delete to an artifact |
-| `UploadArtifact` | Ability to create an artifact. Typically just for debugging; agents have this access for a particular session. |
-| `DownloadArtifact` | Ability to download an artifact |
+| `ReadArtifact` |  |
+| `WriteArtifact` |  |
+| `DeleteArtifact` |  |
+| `UploadArtifact` |  |
+| `DownloadArtifact` |  |
 
 ## Bisect
 
 | Name | Description |
 | ---- | ----------- |
-| `CreateBisectTask` | Ability to start new bisect tasks |
-| `UpdateBisectTask` | Ability to update a bisect task |
-| `ViewBisectTask` | Ability to view a bisect task |
+| `CreateBisectTask` |  |
+| `UpdateBisectTask` |  |
+| `ViewBisectTask` |  |
 
 ## BuildHealth
 
 | Name | Description |
 | ---- | ----------- |
-| `AddUpdateDeleteBuildHealthFilter` | Ability to add, update, and delete build health filters. |
+| `AddUpdateDeleteBuildHealthFilter` |  |
 
+## Compute
+
+| Name | Description |
+| ---- | ----------- |
+| `AddComputeTasks` |  |
+| `GetComputeTasks` |  |
+| `UbaCacheRead` |  |
+| `UbaCacheWrite` |  |
+
 ## Devices
 
 | Name | Description |
 | ---- | ----------- |
-| `DeviceRead` | Ability to read devices |
-| `DeviceWrite` | Ability to write devices |
+| `DeviceRead` |  |
+| `DeviceWrite` |  |
 
 ## Jobs
 
 | Name | Description |
 | ---- | ----------- |
-| `CreateJob` | Ability to start new jobs |
-| `UpdateJob` | Rename a job, modify its priority, etc... |
-| `DeleteJob` | Delete a job properties |
-| `ExecuteJob` | Allows updating a job metadata (name, changelist number, step properties, new groups, job states, etc...). Typically granted to agents. Not user facing. |
-| `RetryJobStep` | Ability to retry a failed job step |
-| `ViewJob` | Ability to view a job |
+| `CreateJob` |  |
+| `UpdateJob` |  |
+| `DeleteJob` |  |
+| `ExecuteJob` |  |
+| `RetryJobStep` |  |
+| `ViewJob` |  |
 
+## Leases
+
+| Name | Description |
+| ---- | ----------- |
+| `ViewLeases` |  |
+| `ViewLeaseTasks` |  |
+
+## Logs
+
+| Name | Description |
+| ---- | ----------- |
+| `CreateLog` |  |
+| `UpdateLog` |  |
+| `ViewLog` |  |
+| `WriteLogData` |  |
+| `CreateEvent` |  |
+| `ViewEvent` |  |
+
 ## Notices
 
@@ -284,23 +1104,33 @@
 
 | Name | Description |
 | ---- | ----------- |
-| `CreateSubscription` | Ability to subscribe to notifications |
+| `CreateSubscription` |  |
 
+## Pools
+
+| Name | Description |
+| ---- | ----------- |
+| `CreatePool` |  |
+| `UpdatePool` |  |
+| `DeletePool` |  |
+| `ViewPool` |  |
+| `ListPools` |  |
+
 ## Projects
 
 | Name | Description |
 | ---- | ----------- |
-| `CreateProject` | Allows the creation of new projects |
-| `DeleteProject` | Allows deletion of projects. |
-| `UpdateProject` | Modify attributes of a project (name, categories, etc...) |
-| `ViewProject` | View information about a project |
+| `CreateProject` |  |
+| `DeleteProject` |  |
+| `UpdateProject` |  |
+| `ViewProject` |  |
 
 ## Replicators
 
 | Name | Description |
 | ---- | ----------- |
-| `UpdateReplicator` | Allows deletion of projects. |
-| `ViewReplicator` | Allows the creation of new projects |
+| `UpdateReplicator` |  |
+| `ViewReplicator` |  |
 
 ## Secrets
 
@@ -317,36 +1147,51 @@
 | `DeleteAccount` |  |
 | `ViewAccount` |  |
 
+## Sessions
+
+| Name | Description |
+| ---- | ----------- |
+| `CreateSession` |  |
+| `ViewSession` |  |
+
+## Software
+
+| Name | Description |
+| ---- | ----------- |
+| `UploadSoftware` |  |
+| `DownloadSoftware` |  |
+| `DeleteSoftware` |  |
+
+## Storage
+
+| Name | Description |
+| ---- | ----------- |
+| `ReadBlobs` | Ability to read blobs from the storage service |
+| `WriteBlobs` | Ability to write blobs to the storage service |
+| `ReadAliases` | Ability to read aliases from the storage service |
+| `WriteAliases` | Ability to write aliases to the storage service |
+| `DeleteAliases` | Ability to write aliases to the storage service |
+| `ReadRefs` | Ability to read refs from the storage service |
+| `WriteRefs` | Ability to write refs to the storage service |
+| `DeleteRefs` | Ability to delete refs |
+
 ## Streams
 
 | Name | Description |
 | ---- | ----------- |
-| `CreateStream` | Allows the creation of new streams within a project |
-| `UpdateStream` | Allows updating a stream (agent types, templates, schedules) |
-| `DeleteStream` | Allows deleting a stream |
-| `ViewStream` | Ability to view a stream |
-| `ViewChanges` | View changes submitted to a stream. NOTE: this returns responses from the server's Perforce account, which may be a priviledged user. |
-| `ViewTemplate` | View template associated with a stream |
+| `CreateStream` |  |
+| `UpdateStream` |  |
+| `DeleteStream` |  |
+| `ViewStream` |  |
+| `ViewChanges` |  |
+| `ViewTemplate` |  |
 
-## Symbols
-
-| Name | Description |
-| ---- | ----------- |
-| `ReadSymbols` |  |
-
 ## TestData
 
 | Name | Description |
 | ---- | ----------- |
-| `TestHealthRead` | Ability to read test health |
-| `TestHealthWrite` | Ability to write test health |
+| `TestHealthRead` |  |
+| `TestHealthWrite` |  |
 
-## Tools
-
-| Name | Description |
-| ---- | ----------- |
-| `DownloadTool` |  |
-| `UploadTool` |  |
-
 ## Accounts
 
@@ -361,12 +1206,12 @@
 
 | Name | Description |
 | ---- | ----------- |
-| `CreateAgent` | Ability to create an agent. This may be done explicitly, or granted to agents to allow them to self-register. |
-| `CreateWorkstationAgent` | Ability to create a workstation agent. Granted to users with hardware to offer. |
-| `UpdateAgent` | Update an agent's name, pools, etc... |
-| `DeleteAgent` | Soft-delete an agent |
-| `ViewAgent` | View an agent |
-| `ListAgents` | List the available agents |
+| `CreateAgent` |  |
+| `CreateWorkstationAgent` |  |
+| `UpdateAgent` |  |
+| `DeleteAgent` |  |
+| `ViewAgent` |  |
+| `ListAgents` |  |
 
 ## Artifacts
 
@@ -396,10 +1241,10 @@
 
 | Name | Description |
 | ---- | ----------- |
-| `AddComputeTasks` | User can add tasks to the compute cluster |
-| `GetComputeTasks` | User can get and list tasks from the compute cluster |
-| `UbaCacheRead` | User can read from Unreal Build Accelerator cache |
-| `UbaCacheWrite` | User can write to Unreal Build Accelerator cache |
+| `AddComputeTasks` |  |
+| `GetComputeTasks` |  |
+| `UbaCacheRead` |  |
+| `UbaCacheWrite` |  |
 
 ## Devices
 
@@ -423,19 +1268,19 @@
 
 | Name | Description |
 | ---- | ----------- |
-| `ViewLeases` | View all the leases that an agent has worked on |
-| `ViewLeaseTasks` | View the task data for a lease |
+| `ViewLeases` |  |
+| `ViewLeaseTasks` |  |
 
 ## Logs
 
 | Name | Description |
 | ---- | ----------- |
-| `CreateLog` | Ability to create a log. Implicitly granted to agents. |
-| `UpdateLog` | Ability to update log metadata |
-| `ViewLog` | Ability to view a log contents |
-| `WriteLogData` | Ability to write log data |
-| `CreateEvent` | Ability to create events |
-| `ViewEvent` | Ability to view events |
+| `CreateLog` |  |
+| `UpdateLog` |  |
+| `ViewLog` |  |
+| `WriteLogData` |  |
+| `CreateEvent` |  |
+| `ViewEvent` |  |
 
 ## Notices
 
@@ -455,11 +1300,11 @@
 
 | Name | Description |
 | ---- | ----------- |
-| `CreatePool` | Create a global pool of agents |
-| `UpdatePool` | Modify an agent pool |
-| `DeletePool` | Delete an agent pool |
-| `ViewPool` | Ability to view a pool |
-| `ListPools` | View all the available agent pools |
+| `CreatePool` |  |
+| `UpdatePool` |  |
+| `DeletePool` |  |
+| `ViewPool` |  |
+| `ListPools` |  |
 
 ## Projects
 
@@ -496,16 +1341,29 @@
 
 | Name | Description |
 | ---- | ----------- |
-| `CreateSession` | Granted to agents to call CreateSession, which returns a bearer token identifying themselves valid to call UpdateSesssion via gRPC. |
-| `ViewSession` | Allows viewing information about an agent session |
+| `CreateSession` |  |
+| `ViewSession` |  |
 
 ## Software
 
 | Name | Description |
 | ---- | ----------- |
-| `UploadSoftware` | Ability to upload new versions of the agent software |
-| `DownloadSoftware` | Ability to download the agent software |
-| `DeleteSoftware` | Ability to delete agent software |
+| `UploadSoftware` |  |
+| `DownloadSoftware` |  |
+| `DeleteSoftware` |  |
 
+## Storage
+
+| Name | Description |
+| ---- | ----------- |
+| `ReadBlobs` |  |
+| `WriteBlobs` |  |
+| `ReadAliases` |  |
+| `WriteAliases` |  |
+| `DeleteAliases` |  |
+| `ReadRefs` |  |
+| `WriteRefs` |  |
+| `DeleteRefs` |  |
+
 ## Streams
 
@@ -522,7 +1380,7 @@
 
 | Name | Description |
 | ---- | ----------- |
-| `ReadSymbols` |  |
+| `ReadSymbols` | Ability to download symbols |
 
 ## TestData
 
@@ -531,13 +1389,6 @@
 | `TestHealthRead` |  |
 | `TestHealthWrite` |  |
 
-## Tools
-
-| Name | Description |
-| ---- | ----------- |
-| `DownloadTool` |  |
-| `UploadTool` |  |
-
 ## Accounts
 
 | Name | Description |
@@ -701,14 +1552,14 @@
 
 | Name | Description |
 | ---- | ----------- |
-| `ReadBlobs` | Ability to read blobs from the storage service |
-| `WriteBlobs` | Ability to write blobs to the storage service |
-| `ReadAliases` | Ability to read aliases from the storage service |
-| `WriteAliases` | Ability to write aliases to the storage service |
-| `DeleteAliases` | Ability to write aliases to the storage service |
-| `ReadRefs` | Ability to read refs from the storage service |
-| `WriteRefs` | Ability to write refs to the storage service |
-| `DeleteRefs` | Ability to delete refs |
+| `ReadBlobs` |  |
+| `WriteBlobs` |  |
+| `ReadAliases` |  |
+| `WriteAliases` |  |
+| `DeleteAliases` |  |
+| `ReadRefs` |  |
+| `WriteRefs` |  |
+| `DeleteRefs` |  |
 
 ## Streams
 
@@ -738,8 +1589,8 @@
 
 | Name | Description |
 | ---- | ----------- |
-| `DownloadTool` |  |
-| `UploadTool` |  |
+| `DownloadTool` | Ability to download a tool |
+| `UploadTool` | Ability to upload new tool versions |
 
 ## Accounts
 
diff --git a/Engine/Source/Programs/Horde/Docs/Config/Schema/Globals.md b/Engine/Source/Programs/Horde/Docs/Config/Schema/Globals.md
--- a/Engine/Source/Programs/Horde/Docs/Config/Schema/Globals.md
+++ b/Engine/Source/Programs/Horde/Docs/Config/Schema/Globals.md
@@ -59,210 +59,15 @@
 
 Name | Description
 ---- | -----------
+`build` | [BuildConfig](#buildconfig)<br>Configuration for the build plugin
+`compute` | [ComputeConfig](#computeconfig)<br>Configuration for the compute plugin
 `analytics` | [AnalyticsConfig](#analyticsconfig)<br>Configuration for the analytics plugin
-`exampleplugin` | [ExamplePluginConfig](#examplepluginconfig)<br>Configuration for the exampleplugin plugin
 `experimental` | [ExperimentalConfig](#experimentalconfig)<br>Configuration for the experimental plugin
 `secrets` | [SecretsConfig](#secretsconfig)<br>Configuration for the secrets plugin
+`exampleplugin` | [ExamplePluginConfig](#examplepluginconfig)<br>Configuration for the exampleplugin plugin
+`storage` | [StorageConfig](#storageconfig)<br>Configuration for the storage plugin
 `symbols` | [SymbolsConfig](#symbolsconfig)<br>Configuration for the symbols plugin
 `tools` | [ToolsConfig](#toolsconfig)<br>Configuration for the tools plugin
-`build` | [BuildConfig](#buildconfig)<br>Configuration for the build plugin
-`compute` | [ComputeConfig](#computeconfig)<br>Configuration for the compute plugin
-`storage` | [StorageConfig](#storageconfig)<br>Configuration for the storage plugin
 
-## AnalyticsConfig
-
-Config settings for analytics
-
-Name | Description
----- | -----------
-`stores` | [TelemetryStoreConfig](Telemetry.md)`[]`<br>Metrics to aggregate on the Horde server
-
-## ExamplePluginConfig
-
-Configuration for the ExamplePlugin plugin
-
-Name | Description
----- | -----------
-`greetingMessage` | `string`<br>The Greeting provided by the example plugin
-
-## ExperimentalConfig
-
-Configuration for the experimental plugin
-
-Name | Description
----- | -----------
-`notifications` | [NotificationConfig](Notification.md)`[]`<br>Notifications to be sent on the Horde server
-
-## SecretsConfig
-
-Configuration for the secrets system
-
-Name | Description
----- | -----------
-`secrets` | [SecretConfig](#secretconfig)`[]`<br>List of secrets
-`providerConfigs` | [SecretProviderConfig](#secretproviderconfig)`[]`<br>List of configurations for secret providers
-
-## SecretConfig
-
-Configuration for a secret value
-
-Name | Description
----- | -----------
-`id` | `string`<br>Identifier for this secret
-`data` | `string` `->` `string`<br>Key/value pairs associated with this secret
-`sources` | [ExternalSecretConfig](#externalsecretconfig)`[]`<br>Providers to source key/value pairs from
-`acl` | [AclConfig](#aclconfig)<br>Defines access to this particular secret
-
-## ExternalSecretConfig
-
-Configuration for an external secret provider
-
-Name | Description
----- | -----------
-`providerConfig` | `string`<br>The name to look up a secret provider config. A config is mutually exclusive to and to be used when a secret requires additional options than just the name of the provider
-`provider` | `string`<br>Name of the provider to use
-`format` | [ExternalSecretFormat](#externalsecretformat-enum)<br>Format of the secret
-`key` | `string`<br>Optional key indicating the parameter to set in the resulting data array. Required if  is .
-`path` | `string`<br>Optional value indicating what to fetch from the provider
-
-## ExternalSecretFormat (Enum)
-
-Format describing how to parse external secret values
-
-Name | Description
----- | -----------
-`Text` | Secret is a plain text value which will be stored using the external secret key
-`Json` | Secret is a JSON formatted string containing key/value pairs
-
-## AclConfig
-
-Parameters to update an ACL
-
-Name | Description
----- | -----------
-`entries` | [AclEntryConfig](#aclentryconfig)`[]`<br>Entries to replace the existing ACL
-`profiles` | [AclProfileConfig](#aclprofileconfig)`[]`<br>Defines profiles which allow grouping sets of actions into named collections
-`inherit` | `boolean`<br>Whether to inherit permissions from the parent ACL
-`exceptions` | `string[]`<br>List of exceptions to the inherited setting
-
-## AclEntryConfig
-
-Individual entry in an ACL
-
-Name | Description
----- | -----------
-`claim` | [AclClaimConfig](#aclclaimconfig)<br>Name of the user or group
-`actions` | `string[]`<br>Array of actions to allow
-`profiles` | `string[]`<br>List of profiles to grant
-
-## AclClaimConfig
-
-New claim to create
-
-Name | Description
----- | -----------
-`type` | `string`<br>The claim type
-`value` | `string`<br>The claim value
-
-## AclProfileConfig
-
-Configuration for an ACL profile. This defines a preset group of actions which can be given to a user via an ACL entry.
-
-Name | Description
----- | -----------
-`id` | `string`<br>Identifier for this profile
-`actions` | `string[]`<br>Actions to include
-`excludeActions` | `string[]`<br>Actions to exclude from the inherited actions
-`extends` | `string[]`<br>Other profiles to extend from
-
-## SecretProviderConfig
-
-Common settings object for different secret providers
-
-Name | Description
----- | -----------
-`name` | `string`<br>Name of the secret provider config
-`provider` | `string`<br>Name of the provider
-`hcpVault` | [HcpVaultConfig](#hcpvaultconfig)<br>Configuration for HCP Vault
-`azureKeyVault` | [AzureKeyVaultConfig](#azurekeyvaultconfig)<br>Configuration for Azure Key Vault
-
-## HcpVaultConfig
-
-Configuration for HCP Vault secrets.
-
-Name | Description
----- | -----------
-`preSharedKey` | `string`<br>The Vault authentication token to use when the type of credentials is .<br>More information on this token can be found in the API documentation for X-Vault-Token
-`endPoint` | `string`<br>The address of the Vault server.
-`awsArnRole` | `string`<br>The optional AWS ARN to assume a role when logging into Vault. To be used when the type of credentials is . When not set the AWS SDK default credential search will be used.
-`awsIamServerId` | `string`<br>The value for the X-Vault-AWS-IAM-Server-ID header when the Vault server is configured with iam_server_id_header_value. Most likely to be the Vault server's DNS name. To be used when the type of credentials is .
-`role` | `string`<br>The name of the role in the Vault configuration for the login. This role name can be different to the name of the AWS IAM principal. To be used when the type of credentials is . When not set the name of IAM principal will be used.
-`credentials` | [HcpVaultCredentialsType](#hcpvaultcredentialstype-enum)<br>Type of credentials to use.
-
-## HcpVaultCredentialsType (Enum)
-
-Credentials to use to log into HCP Vault.
-
-Name | Description
----- | -----------
-`PreSharedKey` | Use a pre-shared token provided by Vault for authentication.
-`AwsAuth` | Use AWS IAM credentials to authenticate against the Vault server.
-
-## AzureKeyVaultConfig
-
-Configuration for Azure Key Vault secrets.
-
-Name | Description
----- | -----------
-`vaultUri` | `string`<br>The Uri to the Azure Key Vault.
-
-## SymbolsConfig
-
-Configuration for the tools system
-
-Name | Description
----- | -----------
-`stores` | [SymbolStoreConfig](#symbolstoreconfig)`[]`<br>List of symbol stores
-
-## SymbolStoreConfig
-
-Configuration for a symbol store
-
-Name | Description
----- | -----------
-`id` | `string`<br>Identifier for this store
-`namespaceId` | `string`<br>Configuration for the symbol store backend
-`public` | `boolean`<br>Whether to make this store available without auth
-`acl` | [AclConfig](#aclconfig)<br>Access to the symbol store
-
-## ToolsConfig
-
-Configuration for the tools system
-
-Name | Description
----- | -----------
-`tools` | [ToolConfig](#toolconfig)`[]`<br>Tool configurations
-
-## ToolConfig
-
-Options for configuring a tool
-
-Name | Description
----- | -----------
-`id` | `string`<br>Unique identifier for the tool
-`name` | `string`<br>Name of the tool
-`description` | `string`<br>Description for the tool
-`category` | `string`<br>Category for the tool. Will cause the tool to be shown in a different tab in the dashboard.
-`group` | `string`<br>Grouping key for different variations of the same tool. The dashboard will show these together.
-`platforms` | `string[]`<br>Platforms for this tool. Takes the form of a NET RID (https://learn.microsoft.com/en-us/dotnet/core/rid-catalog).
-`public` | `boolean`<br>Whether this tool should be exposed for download on a public endpoint without authentication
-`showInUgs` | `boolean`<br>Whether to show this tool for download in the UGS tools menu
-`showInDashboard` | `boolean`<br>Whether to show this tool for download in the dashboard
-`showInToolbox` | `boolean`<br>Whether to show this tool for download in Unreal Toolbox
-`metadata` | `string` `->` `string`<br>Metadata for this tool
-`namespaceId` | `string`<br>Default namespace for new deployments of this tool
-`maxDeployments` | `integer`<br>Maximum number of deployments to retain for this tool. When exceeded, the oldest deployments are removed.
-`acl` | [AclConfig](#aclconfig)<br>Permissions for the tool
-
 ## BuildConfig
 
@@ -282,6 +87,7 @@
 `issueFixedTag` | `string`<br>Commit tag to use for marking issues as fixed
 `issues` | [IssueConfig](Issues.md)<br>Build health issue configurations
 `enableIssueRootCauseAnalysisNotifications` | `boolean`<br>Whether to allow notifications for Root Cause Analysis to be sent
+`enableAutoRcaSlackNotifications` | `boolean`<br>Global kill switch for posting auto-RCA results to Slack triage threads. When false (default), no auto-RCA Slack notifications are sent regardless of per-workflow AutoRcaSlackEnabled settings.
 `autoRcaUseMultiStage` | `boolean`<br>Whether to use multi-stage LLM analysis for auto-RCA. Multi-stage is slower but potentially more accurate. Set to false for faster processing with higher throughput.
 `maxConcurrentRcaRequests` | `integer`<br>Maximum number of concurrent auto-RCA requests. Limits load on the external LLM API. Set to 0 to disable auto-RCA globally (acts as a kill switch).
 `issueFindSuspectsMaxSuccessDays` | `integer`<br>The maximum days since node success to scan for suspects
@@ -410,6 +216,47 @@
 `keepDays` | `integer`<br>Number of days to retain artifacts of this type
 `namespaceId` | `string`<br>Storage namespace to use for this artifact types
 
+## AclConfig
+
+Parameters to update an ACL
+
+Name | Description
+---- | -----------
+`entries` | [AclEntryConfig](#aclentryconfig)`[]`<br>Entries to replace the existing ACL
+`profiles` | [AclProfileConfig](#aclprofileconfig)`[]`<br>Defines profiles which allow grouping sets of actions into named collections
+`inherit` | `boolean`<br>Whether to inherit permissions from the parent ACL
+`exceptions` | `string[]`<br>List of exceptions to the inherited setting
+
+## AclEntryConfig
+
+Individual entry in an ACL
+
+Name | Description
+---- | -----------
+`claim` | [AclClaimConfig](#aclclaimconfig)<br>Name of the user or group
+`actions` | `string[]`<br>Array of actions to allow
+`profiles` | `string[]`<br>List of profiles to grant
+
+## AclClaimConfig
+
+New claim to create
+
+Name | Description
+---- | -----------
+`type` | `string`<br>The claim type
+`value` | `string`<br>The claim value
+
+## AclProfileConfig
+
+Configuration for an ACL profile. This defines a preset group of actions which can be given to a user via an ACL entry.
+
+Name | Description
+---- | -----------
+`id` | `string`<br>Identifier for this profile
+`actions` | `string[]`<br>Actions to include
+`excludeActions` | `string[]`<br>Actions to exclude from the inherited actions
+`extends` | `string[]`<br>Other profiles to extend from
+
 ## ComputeConfig
 
 Configuration for the compute system
@@ -609,6 +456,112 @@
 `description` | `string`<br>Human-readable description
 `computeId` | `string`<br>Compute ID for this network (used when allocating compute resources)
 
+## AnalyticsConfig
+
+Config settings for analytics
+
+Name | Description
+---- | -----------
+`stores` | [TelemetryStoreConfig](Telemetry.md)`[]`<br>Metrics to aggregate on the Horde server
+
+## ExperimentalConfig
+
+Configuration for the experimental plugin
+
+Name | Description
+---- | -----------
+`notifications` | [NotificationConfig](Notification.md)`[]`<br>Notifications to be sent on the Horde server
+
+## SecretsConfig
+
+Configuration for the secrets system
+
+Name | Description
+---- | -----------
+`secrets` | [SecretConfig](#secretconfig)`[]`<br>List of secrets
+`providerConfigs` | [SecretProviderConfig](#secretproviderconfig)`[]`<br>List of configurations for secret providers
+
+## SecretConfig
+
+Configuration for a secret value
+
+Name | Description
+---- | -----------
+`id` | `string`<br>Identifier for this secret
+`data` | `string` `->` `string`<br>Key/value pairs associated with this secret
+`sources` | [ExternalSecretConfig](#externalsecretconfig)`[]`<br>Providers to source key/value pairs from
+`acl` | [AclConfig](#aclconfig)<br>Defines access to this particular secret
+
+## ExternalSecretConfig
+
+Configuration for an external secret provider
+
+Name | Description
+---- | -----------
+`providerConfig` | `string`<br>The name to look up a secret provider config. A config is mutually exclusive to and to be used when a secret requires additional options than just the name of the provider
+`provider` | `string`<br>Name of the provider to use
+`format` | [ExternalSecretFormat](#externalsecretformat-enum)<br>Format of the secret
+`key` | `string`<br>Optional key indicating the parameter to set in the resulting data array. Required if  is .
+`path` | `string`<br>Optional value indicating what to fetch from the provider
+
+## ExternalSecretFormat (Enum)
+
+Format describing how to parse external secret values
+
+Name | Description
+---- | -----------
+`Text` | Secret is a plain text value which will be stored using the external secret key
+`Json` | Secret is a JSON formatted string containing key/value pairs
+
+## SecretProviderConfig
+
+Common settings object for different secret providers
+
+Name | Description
+---- | -----------
+`name` | `string`<br>Name of the secret provider config
+`provider` | `string`<br>Name of the provider
+`hcpVault` | [HcpVaultConfig](#hcpvaultconfig)<br>Configuration for HCP Vault
+`azureKeyVault` | [AzureKeyVaultConfig](#azurekeyvaultconfig)<br>Configuration for Azure Key Vault
+
+## HcpVaultConfig
+
+Configuration for HCP Vault secrets.
+
+Name | Description
+---- | -----------
+`preSharedKey` | `string`<br>The Vault authentication token to use when the type of credentials is .<br>More information on this token can be found in the API documentation for X-Vault-Token
+`endPoint` | `string`<br>The address of the Vault server.
+`awsArnRole` | `string`<br>The optional AWS ARN to assume a role when logging into Vault. To be used when the type of credentials is . When not set the AWS SDK default credential search will be used.
+`awsIamServerId` | `string`<br>The value for the X-Vault-AWS-IAM-Server-ID header when the Vault server is configured with iam_server_id_header_value. Most likely to be the Vault server's DNS name. To be used when the type of credentials is .
+`role` | `string`<br>The name of the role in the Vault configuration for the login. This role name can be different to the name of the AWS IAM principal. To be used when the type of credentials is . When not set the name of IAM principal will be used.
+`credentials` | [HcpVaultCredentialsType](#hcpvaultcredentialstype-enum)<br>Type of credentials to use.
+
+## HcpVaultCredentialsType (Enum)
+
+Credentials to use to log into HCP Vault.
+
+Name | Description
+---- | -----------
+`PreSharedKey` | Use a pre-shared token provided by Vault for authentication.
+`AwsAuth` | Use AWS IAM credentials to authenticate against the Vault server.
+
+## AzureKeyVaultConfig
+
+Configuration for Azure Key Vault secrets.
+
+Name | Description
+---- | -----------
+`vaultUri` | `string`<br>The Uri to the Azure Key Vault.
+
+## ExamplePluginConfig
+
+Configuration for the ExamplePlugin plugin
+
+Name | Description
+---- | -----------
+`greetingMessage` | `string`<br>The Greeting provided by the example plugin
+
 ## StorageConfig
 
 Configuration for storage
@@ -686,3 +639,51 @@
 `gcDelayHrs` | `number`<br>How long to keep newly uploaded orphanned objects before allowing them to be deleted, in hours.
 `enableAliases` | `boolean`<br>Support querying exports by their aliases
 `acl` | [AclConfig](#aclconfig)<br>Access list for this namespace
+
+## SymbolsConfig
+
+Configuration for the tools system
+
+Name | Description
+---- | -----------
+`stores` | [SymbolStoreConfig](#symbolstoreconfig)`[]`<br>List of symbol stores
+
+## SymbolStoreConfig
+
+Configuration for a symbol store
+
+Name | Description
+---- | -----------
+`id` | `string`<br>Identifier for this store
+`namespaceId` | `string`<br>Configuration for the symbol store backend
+`public` | `boolean`<br>Whether to make this store available without auth
+`acl` | [AclConfig](#aclconfig)<br>Access to the symbol store
+
+## ToolsConfig
+
+Configuration for the tools system
+
+Name | Description
+---- | -----------
+`tools` | [ToolConfig](#toolconfig)`[]`<br>Tool configurations
+
+## ToolConfig
+
+Options for configuring a tool
+
+Name | Description
+---- | -----------
+`id` | `string`<br>Unique identifier for the tool
+`name` | `string`<br>Name of the tool
+`description` | `string`<br>Description for the tool
+`category` | `string`<br>Category for the tool. Will cause the tool to be shown in a different tab in the dashboard.
+`group` | `string`<br>Grouping key for different variations of the same tool. The dashboard will show these together.
+`platforms` | `string[]`<br>Platforms for this tool. Takes the form of a NET RID (https://learn.microsoft.com/en-us/dotnet/core/rid-catalog).
+`public` | `boolean`<br>Whether this tool should be exposed for download on a public endpoint without authentication
+`showInUgs` | `boolean`<br>Whether to show this tool for download in the UGS tools menu
+`showInDashboard` | `boolean`<br>Whether to show this tool for download in the dashboard
+`showInToolbox` | `boolean`<br>Whether to show this tool for download in Unreal Toolbox
+`metadata` | `string` `->` `string`<br>Metadata for this tool
+`namespaceId` | `string`<br>Default namespace for new deployments of this tool
+`maxDeployments` | `integer`<br>Maximum number of deployments to retain for this tool. When exceeded, the oldest deployments are removed.
+`acl` | [AclConfig](#aclconfig)<br>Permissions for the tool
diff --git a/Engine/Source/Programs/Horde/Docs/Config/Schema/Streams.md b/Engine/Source/Programs/Horde/Docs/Config/Schema/Streams.md
--- a/Engine/Source/Programs/Horde/Docs/Config/Schema/Streams.md
+++ b/Engine/Source/Programs/Horde/Docs/Config/Schema/Streams.md
@@ -512,6 +512,7 @@ Name | Description
 `escalateTimes` | `integer[]`<br>Times after an issue has been opened to escalate to the alias above, in minutes. Continues to notify on the last interval once reaching the end of the list.
 `escalateRootCauseAnalysisTimes` | `integer[]`<br>Times after a manual root cause analysis has been opened to escalate to the owner, in minutes. Continues to notify on the last interval once reaching the end of the list.
 `autoRcaEnabled` | `boolean`<br>Whether to automatically run RCA processing for issues in this workflow
+`autoRcaSlackEnabled` | `boolean`<br>Whether to post auto-RCA results to the issue's Slack triage thread
 `maxMentions` | `integer`<br>Maximum number of people to mention on a triage thread
 `allowMentions` | `boolean`<br>Whether to mention people on this thread. Useful to disable for testing.
 `inviteRestrictedUsers` | `boolean`<br>Uses the admin.conversations.invite API to invite users to the channel
diff --git a/Engine/Source/Programs/Horde/Docs/Deployment/ServerSettings.md b/Engine/Source/Programs/Horde/Docs/Deployment/ServerSettings.md
--- a/Engine/Source/Programs/Horde/Docs/Deployment/ServerSettings.md
+++ b/Engine/Source/Programs/Horde/Docs/Deployment/ServerSettings.md
@@ -127,159 +127,15 @@
 
 Name | Description
 ---- | -----------
+`build` | [BuildServerConfig](#buildserverconfig)<br>Configuration for the build plugin
+`compute` | [ComputeServerConfig](#computeserverconfig)<br>Configuration for the compute plugin
 `analytics` | [AnalyticsServerConfig](#analyticsserverconfig)<br>Configuration for the analytics plugin
-`exampleplugin` | [ExamplePluginServerConfig](#examplepluginserverconfig)<br>Configuration for the exampleplugin plugin
 `experimental` | [ExperimentalServerConfig](#experimentalserverconfig)<br>Configuration for the experimental plugin
 `secrets` | [SecretsServerConfig](#secretsserverconfig)<br>Configuration for the secrets plugin
+`exampleplugin` | [ExamplePluginServerConfig](#examplepluginserverconfig)<br>Configuration for the exampleplugin plugin
+`storage` | [StorageServerConfig](#storageserverconfig)<br>Configuration for the storage plugin
 `symbols` | [PluginServerConfig](#pluginserverconfig)<br>Configuration for the symbols plugin
 `tools` | [ToolsServerConfig](#toolsserverconfig)<br>Configuration for the tools plugin
-`build` | [BuildServerConfig](#buildserverconfig)<br>Configuration for the build plugin
-`compute` | [ComputeServerConfig](#computeserverconfig)<br>Configuration for the compute plugin
-`storage` | [StorageServerConfig](#storageserverconfig)<br>Configuration for the storage plugin
 
-## AnalyticsServerConfig
-
-Server configuration for the analytics system
-
-Name | Description
----- | -----------
-`sinks` | [TelemetrySinkConfig](#telemetrysinkconfig)<br>Settings for the various telemetry sinks
-`enabled` | `boolean`<br>Whether the plugin should be enabled or not
-
-## TelemetrySinkConfig
-
-Telemetry sinks
-
-Name | Description
----- | -----------
-`epic` | [EpicTelemetryConfig](#epictelemetryconfig)<br>Settings for the Epic telemetry sink
-`mongo` | [MongoTelemetryConfig](#mongotelemetryconfig)<br>Settings for the MongoDB telemetry sink
-
-## EpicTelemetryConfig
-
-Configuration for the telemetry sink
-
-Name | Description
----- | -----------
-`url` | `string`<br>Base URL for the telemetry server
-`appId` | `string`<br>Application name to send in the event messages
-`enabled` | `boolean`<br>Whether to enable this sink
-
-## MongoTelemetryConfig
-
-Configuration for the telemetry sink
-
-Name | Description
----- | -----------
-`retainDays` | `number`<br>Number of days worth of telmetry events to keep
-`enabled` | `boolean`<br>Whether to enable this sink
-
-## ExamplePluginServerConfig
-
-Server configuration for the ExamplePlugin
-
-Name | Description
----- | -----------
-`enabled` | `boolean`<br>Whether the plugin should be enabled or not
-
-## ExperimentalServerConfig
-
-Server configuration for the experimental plugin
-
-Name | Description
----- | -----------
-`enabled` | `boolean`<br>Whether the plugin should be enabled or not
-
-## SecretsServerConfig
-
-Static configuration for the secrets plugin
-
-Name | Description
----- | -----------
-`withAws` | `boolean`<br>Whether to enable Amazon Web Services (AWS) specific features
-`enabled` | `boolean`<br>Whether the plugin should be enabled or not
-
-## PluginServerConfig
-
-Base class for plugin server config objects
-
-Name | Description
----- | -----------
-`enabled` | `boolean`<br>Whether the plugin should be enabled or not
-
-## ToolsServerConfig
-
-Server configuration for bundled tools
-
-Name | Description
----- | -----------
-`bundledTools` | [BundledToolConfig](#bundledtoolconfig)`[]`<br>Tools bundled along with the server. Data for each tool can be produced using the 'bundle create' command, and should be stored in the Tools directory.
-`enabled` | `boolean`<br>Whether the plugin should be enabled or not
-
-## BundledToolConfig
-
-Configuration for a tool bundled alongsize the server
-
-Name | Description
----- | -----------
-`version` | `string`<br>Version string for the current tool data
-`refName` | `string`<br>Ref name in the tools directory
-`dataDir` | `string`<br>Directory containing blob data for this tool. If empty, the tools/{id} folder next to the server will be used.
-`id` | `string`<br>Unique identifier for the tool
-`name` | `string`<br>Name of the tool
-`description` | `string`<br>Description for the tool
-`category` | `string`<br>Category for the tool. Will cause the tool to be shown in a different tab in the dashboard.
-`group` | `string`<br>Grouping key for different variations of the same tool. The dashboard will show these together.
-`platforms` | `string[]`<br>Platforms for this tool. Takes the form of a NET RID (https://learn.microsoft.com/en-us/dotnet/core/rid-catalog).
-`public` | `boolean`<br>Whether this tool should be exposed for download on a public endpoint without authentication
-`showInUgs` | `boolean`<br>Whether to show this tool for download in the UGS tools menu
-`showInDashboard` | `boolean`<br>Whether to show this tool for download in the dashboard
-`showInToolbox` | `boolean`<br>Whether to show this tool for download in Unreal Toolbox
-`metadata` | `string` `->` `string`<br>Metadata for this tool
-`namespaceId` | `string`<br>Default namespace for new deployments of this tool
-`maxDeployments` | `integer`<br>Maximum number of deployments to retain for this tool. When exceeded, the oldest deployments are removed.
-`acl` | [AclConfig](#aclconfig)<br>Permissions for the tool
-
-## AclConfig
-
-Parameters to update an ACL
-
-Name | Description
----- | -----------
-`entries` | [AclEntryConfig](#aclentryconfig)`[]`<br>Entries to replace the existing ACL
-`profiles` | [AclProfileConfig](#aclprofileconfig)`[]`<br>Defines profiles which allow grouping sets of actions into named collections
-`inherit` | `boolean`<br>Whether to inherit permissions from the parent ACL
-`exceptions` | `string[]`<br>List of exceptions to the inherited setting
-
-## AclEntryConfig
-
-Individual entry in an ACL
-
-Name | Description
----- | -----------
-`claim` | [AclClaimConfig](#aclclaimconfig)<br>Name of the user or group
-`actions` | `string[]`<br>Array of actions to allow
-`profiles` | `string[]`<br>List of profiles to grant
-
-## AclClaimConfig
-
-New claim to create
-
-Name | Description
----- | -----------
-`type` | `string`<br>The claim type
-`value` | `string`<br>The claim value
-
-## AclProfileConfig
-
-Configuration for an ACL profile. This defines a preset group of actions which can be given to a user via an ACL entry.
-
-Name | Description
----- | -----------
-`id` | `string`<br>Identifier for this profile
-`actions` | `string[]`<br>Actions to include
-`excludeActions` | `string[]`<br>Actions to exclude from the inherited actions
-`extends` | `string[]`<br>Other profiles to extend from
-
 ## BuildServerConfig
 
@@ -471,6 +327,68 @@
 `ComputeQueueAwsMetric` | A no-op strategy that reports metrics to let an external AWS auto-scaling policy scale the fleet
 `LeaseUtilizationAwsMetric` | A no-op strategy that reports metrics to let an external AWS auto-scaling policy scale the fleet
 
+## AnalyticsServerConfig
+
+Server configuration for the analytics system
+
+Name | Description
+---- | -----------
+`sinks` | [TelemetrySinkConfig](#telemetrysinkconfig)<br>Settings for the various telemetry sinks
+`enabled` | `boolean`<br>Whether the plugin should be enabled or not
+
+## TelemetrySinkConfig
+
+Telemetry sinks
+
+Name | Description
+---- | -----------
+`epic` | [EpicTelemetryConfig](#epictelemetryconfig)<br>Settings for the Epic telemetry sink
+`mongo` | [MongoTelemetryConfig](#mongotelemetryconfig)<br>Settings for the MongoDB telemetry sink
+
+## EpicTelemetryConfig
+
+Configuration for the telemetry sink
+
+Name | Description
+---- | -----------
+`url` | `string`<br>Base URL for the telemetry server
+`appId` | `string`<br>Application name to send in the event messages
+`enabled` | `boolean`<br>Whether to enable this sink
+
+## MongoTelemetryConfig
+
+Configuration for the telemetry sink
+
+Name | Description
+---- | -----------
+`retainDays` | `number`<br>Number of days worth of telmetry events to keep
+`enabled` | `boolean`<br>Whether to enable this sink
+
+## ExperimentalServerConfig
+
+Server configuration for the experimental plugin
+
+Name | Description
+---- | -----------
+`enabled` | `boolean`<br>Whether the plugin should be enabled or not
+
+## SecretsServerConfig
+
+Static configuration for the secrets plugin
+
+Name | Description
+---- | -----------
+`withAws` | `boolean`<br>Whether to enable Amazon Web Services (AWS) specific features
+`enabled` | `boolean`<br>Whether the plugin should be enabled or not
+
+## ExamplePluginServerConfig
+
+Server configuration for the ExamplePlugin
+
+Name | Description
+---- | -----------
+`enabled` | `boolean`<br>Whether the plugin should be enabled or not
+
 ## StorageServerConfig
 
 Static settings for the storage system
@@ -529,3 +447,85 @@
 `Profile` | Read credentials from the  profile in the AWS config file
 `AssumeRole` | Assume a particular role. Should specify ARN in
 `AssumeRoleWebIdentity` | Assume a particular role using the current environment variables.
+
+## PluginServerConfig
+
+Base class for plugin server config objects
+
+Name | Description
+---- | -----------
+`enabled` | `boolean`<br>Whether the plugin should be enabled or not
+
+## ToolsServerConfig
+
+Server configuration for bundled tools
+
+Name | Description
+---- | -----------
+`bundledTools` | [BundledToolConfig](#bundledtoolconfig)`[]`<br>Tools bundled along with the server. Data for each tool can be produced using the 'bundle create' command, and should be stored in the Tools directory.
+`enabled` | `boolean`<br>Whether the plugin should be enabled or not
+
+## BundledToolConfig
+
+Configuration for a tool bundled alongsize the server
+
+Name | Description
+---- | -----------
+`version` | `string`<br>Version string for the current tool data
+`refName` | `string`<br>Ref name in the tools directory
+`dataDir` | `string`<br>Directory containing blob data for this tool. If empty, the tools/{id} folder next to the server will be used.
+`id` | `string`<br>Unique identifier for the tool
+`name` | `string`<br>Name of the tool
+`description` | `string`<br>Description for the tool
+`category` | `string`<br>Category for the tool. Will cause the tool to be shown in a different tab in the dashboard.
+`group` | `string`<br>Grouping key for different variations of the same tool. The dashboard will show these together.
+`platforms` | `string[]`<br>Platforms for this tool. Takes the form of a NET RID (https://learn.microsoft.com/en-us/dotnet/core/rid-catalog).
+`public` | `boolean`<br>Whether this tool should be exposed for download on a public endpoint without authentication
+`showInUgs` | `boolean`<br>Whether to show this tool for download in the UGS tools menu
+`showInDashboard` | `boolean`<br>Whether to show this tool for download in the dashboard
+`showInToolbox` | `boolean`<br>Whether to show this tool for download in Unreal Toolbox
+`metadata` | `string` `->` `string`<br>Metadata for this tool
+`namespaceId` | `string`<br>Default namespace for new deployments of this tool
+`maxDeployments` | `integer`<br>Maximum number of deployments to retain for this tool. When exceeded, the oldest deployments are removed.
+`acl` | [AclConfig](#aclconfig)<br>Permissions for the tool
+
+## AclConfig
+
+Parameters to update an ACL
+
+Name | Description
+---- | -----------
+`entries` | [AclEntryConfig](#aclentryconfig)`[]`<br>Entries to replace the existing ACL
+`profiles` | [AclProfileConfig](#aclprofileconfig)`[]`<br>Defines profiles which allow grouping sets of actions into named collections
+`inherit` | `boolean`<br>Whether to inherit permissions from the parent ACL
+`exceptions` | `string[]`<br>List of exceptions to the inherited setting
+
+## AclEntryConfig
+
+Individual entry in an ACL
+
+Name | Description
+---- | -----------
+`claim` | [AclClaimConfig](#aclclaimconfig)<br>Name of the user or group
+`actions` | `string[]`<br>Array of actions to allow
+`profiles` | `string[]`<br>List of profiles to grant
+
+## AclClaimConfig
+
+New claim to create
+
+Name | Description
+---- | -----------
+`type` | `string`<br>The claim type
+`value` | `string`<br>The claim value
+
+## AclProfileConfig
+
+Configuration for an ACL profile. This defines a preset group of actions which can be given to a user via an ACL entry.
+
+Name | Description
+---- | -----------
+`id` | `string`<br>Identifier for this profile
+`actions` | `string[]`<br>Actions to include
+`excludeActions` | `string[]`<br>Actions to exclude from the inherited actions
+`extends` | `string[]`<br>Other profiles to extend from
