Horde Dashboard - Add "Open in Build Storage Tool" button for Cloud Artifacts

diff --git a/Engine/Source/Programs/Horde/HordeDashboard/src/backend/Dashboard.ts b/Engine/Source/Programs/Horde/HordeDashboard/src/backend/Dashboard.ts
--- a/Engine/Source/Programs/Horde/HordeDashboard/src/backend/Dashboard.ts
+++ b/Engine/Source/Programs/Horde/HordeDashboard/src/backend/Dashboard.ts
@@ -26,6 +26,13 @@ export enum WebBrowser {
     Other = "Other"
 }
 
+export enum UserPlatform {
+    Windows = "Windows",
+    Mac = "Mac",
+    Linux = "Linux",
+    Other = "Other"
+}
+
 export class Dashboard {
 
     constructor() {
@@ -203,6 +210,24 @@ export class Dashboard {
 
     }
 
+    get platform(): UserPlatform {
+
+        const uaData = (window.navigator as any).userAgentData?.platform?.toLowerCase() ?? "";
+        if (uaData) {
+            if (uaData === "windows") return UserPlatform.Windows;
+            if (uaData === "macos") return UserPlatform.Mac;
+            if (uaData === "linux") return UserPlatform.Linux;
+        }
+
+        const userAgent = window.navigator.userAgent.toLowerCase();
+        if (userAgent.indexOf("windows") > -1) return UserPlatform.Windows;
+        if (userAgent.indexOf("macintosh") > -1 || userAgent.indexOf("mac os") > -1) return UserPlatform.Mac;
+        if (userAgent.indexOf("linux") > -1) return UserPlatform.Linux;
+        return UserPlatform.Other;
+
+    }
+
+
     get helpEmail(): string | undefined {
         return this.config?.helpEmailAddress;
     }
diff --git a/Engine/Source/Programs/Horde/HordeDashboard/src/components/artifacts/ArtifactButton.tsx b/Engine/Source/Programs/Horde/HordeDashboard/src/components/artifacts/ArtifactButton.tsx
--- a/Engine/Source/Programs/Horde/HordeDashboard/src/components/artifacts/ArtifactButton.tsx
+++ b/Engine/Source/Programs/Horde/HordeDashboard/src/components/artifacts/ArtifactButton.tsx
@@ -3,7 +3,7 @@
 import { IContextualMenuProps, DirectionalHint, PrimaryButton, Modal, Stack, Text, DefaultButton, Checkbox } from "@fluentui/react";
 import backend from "horde/backend";
 import { ArtifactBrowserItem, ArtifactBrowserType, ArtifactSelection, CreateZipRequest, DownloadFormat, GetArtifactResponse } from "horde/backend/Api";
-import dashboard from "horde/backend/Dashboard";
+import dashboard, { UserPlatform } from "horde/backend/Dashboard";
 import { copyToClipboard } from "horde/base/utilities/clipboard";
 import { formatBytes } from "horde/base/utilities/stringUtills";
 import { getHordeStyling } from "horde/styles/Styles";
@@ -85,6 +85,12 @@ export const ArtifactButton: React.FC<{ jobId?: string, stepId?: string, artifac
       }
    }
 
+   function openInBuildStorageTool(zenUri: string) {
+      if (zenUri) {
+         window.location.assign(`ugs://runcustomtool?toolName=build%20storage%20tool&args=RunBuildStorageTool.bat%20-GoToBuild%20${encodeURIComponent(zenUri)}`);
+      }
+   }
+
    function downloadZip() {
 
       if (!artifact) {
@@ -206,6 +212,18 @@ export const ArtifactButton: React.FC<{ jobId?: string, stepId?: string, artifac
                   }
                }
             )
+
+            if (dashboard.platform === UserPlatform.Windows) {
+               downloadProps.items.push(
+                  {
+                     key: 'open_build_storage_tool',
+                     text: 'Open in Build Storage Tool',
+                     onClick: () => {
+                        openInBuildStorageTool(zenUri);
+                     }
+                  }
+               )
+            }
          }
       }
    }
@@ -265,4 +283,4 @@ const ZipDownloadWarning: React.FC<{ onContinue: () => void, onClose: () => void
          </Stack>
       </Stack>
    </Modal>
-}
\ No newline at end of file
+}
