Horde: Fix circular DI dependency causing server startup hang

Change RcaQueueService to accept Lazy<IRcaProcessor?> instead of resolving IRcaProcessor directly at construction time.

diff --git a/Engine/Source/Programs/Horde/Plugins/Build/HordeServer.Build.Tests/BuildTestSetup.cs b/Engine/Source/Programs/Horde/Plugins/Build/HordeServer.Build.Tests/BuildTestSetup.cs
--- a/Engine/Source/Programs/Horde/Plugins/Build/HordeServer.Build.Tests/BuildTestSetup.cs
+++ b/Engine/Source/Programs/Horde/Plugins/Build/HordeServer.Build.Tests/BuildTestSetup.cs
@@ -37,6 +37,7 @@ using HordeServer.Jobs.TestData;
 using HordeServer.Jobs.Timing;
 using HordeServer.Logs;
 using HordeServer.Notifications;
+using HordeServer.RCA;
 using HordeServer.Server;
 using HordeServer.Storage;
 using HordeServer.Streams;
@@ -172,6 +173,7 @@ namespace HordeServer.Tests
 
 			services.AddSingleton<AgentService>();
 			services.AddSingleton(provider => new Lazy<AgentService>(provider.GetRequiredService<AgentService>));
+			services.AddSingleton(sp => new Lazy<IRcaProcessor?>(() => sp.GetService<IRcaProcessor>()));
 			services.AddSingleton<AgentRelayService>();
 			services.AddSingleton<AwsAutoScalingLifecycleService>();
 			services.AddSingleton<FleetService>();
diff --git a/Engine/Source/Programs/Horde/Plugins/Build/HordeServer.Build/BuildPlugin.cs b/Engine/Source/Programs/Horde/Plugins/Build/HordeServer.Build/BuildPlugin.cs
--- a/Engine/Source/Programs/Horde/Plugins/Build/HordeServer.Build/BuildPlugin.cs
+++ b/Engine/Source/Programs/Horde/Plugins/Build/HordeServer.Build/BuildPlugin.cs
@@ -124,6 +124,7 @@ namespace HordeServer
 			services.AddSingleton<IssueService>();
 			services.AddSingleton<RcaQueueService>();
 			services.AddSingleton<IRcaEnqueuer>(sp => sp.GetRequiredService<RcaQueueService>());
+			services.AddSingleton(sp => new Lazy<IRcaProcessor?>(() => sp.GetService<IRcaProcessor>()));
 			services.AddSingleton<JobService>();
 			services.AddSingleton<IAuditLog<JobId>>(sp => sp.GetRequiredService<IAuditLogFactory<JobId>>().Create("Jobs.Log", "JobId"));
 			services.AddSingleton<ILogExtAuthProvider>(sp => sp.GetRequiredService<JobService>());
diff --git a/Engine/Source/Programs/Horde/Plugins/Build/HordeServer.Build/Issues/RcaQueueService.cs b/Engine/Source/Programs/Horde/Plugins/Build/HordeServer.Build/Issues/RcaQueueService.cs
--- a/Engine/Source/Programs/Horde/Plugins/Build/HordeServer.Build/Issues/RcaQueueService.cs
+++ b/Engine/Source/Programs/Horde/Plugins/Build/HordeServer.Build/Issues/RcaQueueService.cs
@@ -39,7 +39,10 @@ namespace HordeServer.Issues
 		static readonly RedisListKey<RcaQueueItem> s_rcaQueue = "rca/queue";
 
 		readonly IRedisService _redisService;
-		readonly IRcaProcessor? _rcaProcessor;
+
+		// Lazy-resolved to break circular singleton dependency:
+		// RcaQueueService → IRcaProcessor → LlmRcaProcessor → IRcaNotifier → SlackNotificationSink → IssueService → IRcaEnqueuer → RcaQueueService
+		readonly Lazy<IRcaProcessor?> _rcaProcessor;
 		readonly IIssueCollection _issueCollection;
 		readonly IOptionsMonitor<BuildConfig> _buildConfig;
 		readonly ITicker _ticker;
@@ -48,7 +51,7 @@ namespace HordeServer.Issues
 		/// <summary>
 		/// Constructor
 		/// </summary>
-		public RcaQueueService(IRedisService redisService, IIssueCollection issueCollection, IClock clock, IOptionsMonitor<BuildConfig> buildConfig, ILogger<RcaQueueService> logger, IRcaProcessor? rcaProcessor = null)
+		public RcaQueueService(IRedisService redisService, IIssueCollection issueCollection, IClock clock, IOptionsMonitor<BuildConfig> buildConfig, ILogger<RcaQueueService> logger, Lazy<IRcaProcessor?> rcaProcessor)
 		{
 			_redisService = redisService;
 			_rcaProcessor = rcaProcessor;
@@ -70,7 +73,7 @@ namespace HordeServer.Issues
 		/// <inheritdoc/>
 		public async Task EnqueueAsync(int issueId, StreamId streamId, WorkflowId? workflowId, CancellationToken cancellationToken)
 		{
-			if (_rcaProcessor == null)
+			if (_rcaProcessor.Value == null)
 			{
 				return;
 			}
@@ -89,7 +92,7 @@ namespace HordeServer.Issues
 
 		async ValueTask ProcessQueueAsync(CancellationToken cancellationToken)
 		{
-			if (_rcaProcessor == null)
+			if (_rcaProcessor.Value == null)
 			{
 				return;
 			}
@@ -156,7 +159,7 @@ namespace HordeServer.Issues
 
 			try
 			{
-				await _rcaProcessor!.ProcessRcaAsync(item.IssueId, streamId, workflowId, cancellationToken);
+				await _rcaProcessor.Value!.ProcessRcaAsync(item.IssueId, streamId, workflowId, cancellationToken);
 			}
 			catch (OperationCanceledException) when (cancellationToken.IsCancellationRequested)
 			{
