[Test Automation Hub] Add support for timeline zoom for high frequency test results

#rnx
#Synced-CL 52004614

diff --git a/Engine/Source/Programs/Horde/HordeDashboard/plugins/test-automation/phaseHistoryGraph.ts b/Engine/Source/Programs/Horde/HordeDashboard/plugins/test-automation/phaseHistoryGraph.ts
--- a/Engine/Source/Programs/Horde/HordeDashboard/plugins/test-automation/phaseHistoryGraph.ts
+++ b/Engine/Source/Programs/Horde/HordeDashboard/plugins/test-automation/phaseHistoryGraph.ts
@@ -3,10 +3,11 @@
 import * as d3 from "d3";
 import { TestDataHandler, TestNameRef, MetadataRef, PhaseSessionResult, TestMetaStatus, consecutivePassesThreshold } from './testData';
 import dashboard from 'horde/backend/Dashboard';
-import { getHumanTime, getShortNiceTime } from 'horde/base/utilities/timeUtils';
+import { getShortNiceTime } from 'horde/base/utilities/timeUtils';
 import { getHordeStyling } from 'horde/styles/Styles';
 import { getPhaseSessionStatusColor,  SessionListSelector } from "./testAutomationCommon";
 import { renderToString } from "react-dom/server";
+import { timelineHandler, Margin, constructTimeline, calculateAverageSessionFrequencyHours } from "./timelineGraph";
 
 // Handle bad "@types/d3" types, fix if addressed upstream
 const _d3 = d3 as any;
@@ -22,8 +23,14 @@ type FingerprintSeries = {
     y: number
 };
 
-const minimumBandLength = 60 * 60 * 2; // 2 hours
+const circleSize = 5;
+const minimumBandLength = circleSize * 2;
 const minimumBandCount = 2;
+const highFrequencyThreshold = 4; // in hours
+
+const width = 1000;
+const bandHeight = 4
+const bandOpacities = [ 0.34, 0.64 ];
 
 export class PhaseHistoryGraph {
 
@@ -49,19 +56,23 @@ export class PhaseHistoryGraph {
         const { modeColors } = getHordeStyling();
 
         this.clear(container);
-
+        this.cleanupCallback();
         this.initData(meta);
 
+        // initialize timeline handle if not set
+        timelineHandler.init(this.test.id, this.handler, width, this.margin, reverse);
+        if (timelineHandler.isTimelineAtExtreme && this.averageFrequencyHours < highFrequencyThreshold) {
+            timelineHandler.rescaleTimeline(this.averageFrequencyHours / highFrequencyThreshold);
+        }
+
         this.phaseSessionId = selectedPhaseSessionId;
 
         const handler = this.handler;
         const sessions = this.sessions;
-        const width = 1000
 
         const X = _d3.map(sessions, (r) => handler.commitIdDates.get(r.commitId)!.getTime() / 1000);
         const Y = _d3.map(sessions, (r) => 0);
 
-        const xDomain = d3.extent(handler.commitIdDates.values() as any, (d:any) => d.getTime() / 1000);
         const yDomain = new _d3.InternSet(Y as any);
 
         const I = d3.range(X.length);
@@ -69,174 +80,170 @@ export class PhaseHistoryGraph {
         const yPadding = 1;
         const height = withoutHeader? 20 : Math.ceil((yDomain.size + yPadding) * 13) + this.margin.top + this.margin.bottom;
 
-        const xRange = reverse? [width - this.margin.right, this.margin.left] : [this.margin.left, width - this.margin.right];
         const yRange = [this.margin.top, height - this.margin.bottom];
-
-        const xScale = _d3.scaleTime(xDomain, xRange);
         const yScale = _d3.scalePoint(yDomain, yRange).round(true).padding(yPadding);
 
-        // group sessions into consecutive runs of the same error fingerprint
-        const fingerprintSeries: FingerprintSeries[] = [];
-        let currentSeries: FingerprintSeries | undefined = undefined;
-        let lastFingerprint: string | undefined = undefined;
-        let consecutivePasses = 0;
-
-        sessions.forEach((session, idx) => {
-            const fingerprint = session.errorFingerprint;
-            const x = X[idx];
-            const y = Y[idx];
-
-            if (!fingerprint) {
-                consecutivePasses++;
-                if (consecutivePasses > consecutivePassesThreshold) {
-                    // reset current series if we had more than consecutivePassesThreshold
-                    currentSeries = undefined;
-                    lastFingerprint = undefined;
+        const getFingerprintSeries = () => {
+            // group sessions into consecutive runs of the same error fingerprint
+            const fingerprintSeries: FingerprintSeries[] = [];
+            let currentSeries: FingerprintSeries | undefined = undefined;
+            let lastFingerprint: string | undefined = undefined;
+            let consecutivePasses = 0;
+
+            sessions.forEach((session, idx) => {
+                const fingerprint = session.errorFingerprint;
+                if (!fingerprint) {
+                    consecutivePasses++;
+                    if (consecutivePasses > consecutivePassesThreshold) {
+                        // reset current series if we had more than consecutivePassesThreshold
+                        currentSeries = undefined;
+                        lastFingerprint = undefined;
+                    }
+                    return;
                 }
-                return;
-            }
 
-            if (fingerprint !== lastFingerprint && (!currentSeries || currentSeries.end !== x)) {
-                // New series started
-                currentSeries = { fingerprint, color: getPhaseSessionStatusColor(session), count: 0, start: x, end: x, y: y };
-                fingerprintSeries.push(currentSeries);
-            }
-            currentSeries!.count++;
-            consecutivePasses = 0;
-            if (x < currentSeries!.start) {
-                currentSeries!.start = x;
-            }
-            lastFingerprint = fingerprint;
-        });
+                const xdate = X[idx];
+                if (xdate < timelineHandler.startDate! || xdate > timelineHandler.endDate!) return;
 
-        const bandHeight = 4
-        const bandOpacities = [ 0.34, 0.64 ];
+                if (fingerprint !== lastFingerprint && (!currentSeries || currentSeries.end !== xdate)) {
+                    // New series started
+                    currentSeries = { fingerprint, color: getPhaseSessionStatusColor(session), count: 0, start: xdate, end: xdate, y: Y[idx] };
+                    fingerprintSeries.push(currentSeries);
+                }
+                currentSeries!.count++;
+                consecutivePasses = 0;
+                if (xdate < currentSeries!.start) {
+                    currentSeries!.start = xdate;
+                }
+                lastFingerprint = fingerprint;
+            });
+
+            return fingerprintSeries;
+        }
 
         const svg = d3.select(container)
             .append("svg")
             .attr("width", width)
             .attr("height", height + 24)
-            .attr("viewBox", [0, 0, width, height + 24] as any)
-
-        // top axis
-        const xAxis = (g: SelectionType) => {
-
-            const dateMin = new Date(xDomain[0]! * 1000);
-            const dateMax = new Date(xDomain[1]! * 1000);
-
-            let ticks: number[] = [];
-            for (const date of d3.timeDays(dateMin, dateMax, 1).reverse()) {
-                ticks.push(date.getTime() / 1000);
-            }
-
-            if (ticks.length > 14) {
-                let nticks = [...ticks];
-                // remove first and last, will be readded 
-                const first = nticks.shift()!;
-                const last = nticks.pop()!;
-
-                const n = Math.floor(nticks.length / 12);
-
-                const rticks: number[] = [];
-                for (let i = 0; i < nticks.length; i = i + n) {
-                    rticks.push(nticks[i]);
-                }
-
-                rticks.unshift(first);
-                rticks.push(last);
-                ticks = rticks;
-            }
-
-            g.attr("transform", `translate(0,16)`)
-                .style("font-family", "Horde Open Sans Regular")
-                .style("font-size", "9px")
-                .call(d3.axisTop(xScale)
-                    .tickValues(ticks)
-                    .tickFormat(d => withoutHeader? "" : getHumanTime(new Date((d as number) * 1000)))
-                    .tickSizeOuter(0))
-                .call(g => g.select(".domain").remove())
-                .call(g => g.selectAll(".tick line")
-                    .attr("stroke-opacity", dashboard.darktheme ? 0.35 : 0.25).clone()
-                    .attr("stroke", dashboard.darktheme ? "#6D6C6B" : "#4D4C4B")
-                    .attr("y2", height - this.margin.bottom)
-                )
+            .attr("viewBox", [0, 0, width, height + 24] as any);
+
+        this.timelineCleanupCallback = constructTimeline(
+            svg,
+            () => {
+                // update error bands
+                errorBands.call(drawErrorBands);
+                // update meta groups
+                metaGroups.call(drawMetaGroups);
+            },
+            height - this.margin.bottom,
+            reverse,
+            withoutHeader
+        );
+
+        const errorBands = svg.append("g")
+            .attr("id", "error-bands");
+
+        const drawErrorBands = (bands: SelectionType) => {
+            // add group for bands below all marks
+            const fingerprintSeries = getFingerprintSeries();
+            bands.selectAll(".error-band")
+                .data(fingerprintSeries.filter(band => band.count >= minimumBandCount && (Math.abs(timelineHandler.timelineXScale!(band.end) - timelineHandler.timelineXScale!(band.start)) >= minimumBandLength)), (band: any) => band.start)
+                .join(
+                    enter => enter.append("g")
+                        .attr("class", "error-band")
+                        .each((band, i, nodes) => {
+                            const bandNode = d3.select(nodes[i]);
+                            const bandOpacity = bandOpacities[i % bandOpacities.length];
+                            const xStart = Math.min(timelineHandler.timelineXScale!(band.start), timelineHandler.timelineXScale!(band.end));
+                            const xEnd = Math.max(timelineHandler.timelineXScale!(band.start), timelineHandler.timelineXScale!(band.end));
+                            const bandWidth = (xEnd - xStart) + 8;
+                            bandNode
+                                .attr("transform", `translate(${xStart - 4.5},${yScale(band.y)! + 16 - bandHeight/2})`);
+                            bandNode.append("rect")
+                                .attr("width", bandWidth)
+                                .attr("height", bandHeight)
+                                .attr("fill", band.color)
+                                .attr("opacity", bandOpacity)
+                                .attr("stroke", "none");
+                            bandNode.append("line")
+                                .attr("x1", 1)
+                                .attr("x2", 1)
+                                .attr("y1", bandHeight)
+                                .attr("y2", bandHeight*3)
+                                .attr("stroke", band.color)
+                                .attr("stroke-width", 2)
+                                .attr("stroke-opacity", bandOpacity);
+                            bandNode.append("text")
+                                .attr("text-anchor", "start")
+                                .style("alignment-baseline", "left")
+                                .style("font-family", "Horde Open Sans Regular")
+                                .style("font-size", 9)
+                                .style("font-weight", 800)
+                                .attr("y", 18)
+                                .attr("x", 3)
+                                .attr("fill", band.color)
+                                .text(`#${band.fingerprint.substring(0,2)}`);
+                        }),
+                    update => update.each((band, i, nodes) => {
+                        const bandNode = d3.select(nodes[i]);
+                        const xStart = Math.min(timelineHandler.timelineXScale!(band.start), timelineHandler.timelineXScale!(band.end));
+                        const xEnd = Math.max(timelineHandler.timelineXScale!(band.start), timelineHandler.timelineXScale!(band.end));
+                        const bandWidth = (xEnd - xStart) + 8;
+                        bandNode.attr("transform", `translate(${xStart - 4.5},${yScale(band.y)! + 16 - bandHeight/2})`);
+                        bandNode.select("rect").attr("width", bandWidth);
+                    }),
+                    exit => exit.remove(),
+                );
         }
-
-        svg.append("g").call(xAxis)
-
-        // add group for bands below all marks
-        svg.append("g")
-            .attr("class", "error-bands")
-            .selectAll()
-            .data(fingerprintSeries.filter(band => band.count >= minimumBandCount && (band.end - band.start >= minimumBandLength)))
-            .enter()
-            .append("g")
-            .each((band, i, nodes) => {
-                const bandNode = d3.select(nodes[i]);
-                const bandOpacity = bandOpacities[i % bandOpacities.length];
-                const xStart = Math.min(xScale(band.start), xScale(band.end));
-                const xEnd = Math.max(xScale(band.start), xScale(band.end));
-                const bandWidth = (xEnd - xStart) + 8;
-                bandNode
-                    .attr("transform", `translate(${xStart - 4.5},${yScale(band.y)! + 16 - bandHeight/2})`);
-                bandNode.append("rect")
-                    .attr("width", bandWidth)
-                    .attr("height", bandHeight)
-                    .attr("fill", band.color)
-                    .attr("opacity", bandOpacity)
-                    .attr("stroke", "none");
-                bandNode.append("line")
-                    .attr("x1", 1)
-                    .attr("x2", 1)
-                    .attr("y1", bandHeight)
-                    .attr("y2", bandHeight*3)
-                    .attr("stroke", band.color)
-                    .attr("stroke-width", 2)
-                    .attr("stroke-opacity", bandOpacity);
-                bandNode.append("text")
-                    .attr("text-anchor", "start")
-                    .style("alignment-baseline", "left")
-                    .style("font-family", "Horde Open Sans Regular")
-                    .style("font-size", 9)
-                    .style("font-weight", 800)
-                    .attr("y", 18)
-                    .attr("x", 3)
-                    .attr("fill", band.color)
-                    .text(`#${band.fingerprint.substring(0,2)}`);
-            });
-
-        const g = svg.append("g")
-            .selectAll()
-            .data(_d3.group(I, i => Y[i]))
-            .join("g")
-            .attr("class", "meta")
-            .attr("id", ([y]: any) => `meta${y as any}`)
-            .attr("transform", ([y]: any) => `translate(0,${(yScale(y) as any) + 16})`);
-
-        g.append("line")
-            .attr("stroke", dashboard.darktheme ? "#6D6C6B" : "#4D4C4B")
-            .attr("stroke-width", 1)
-            .attr("stroke-linecap", 4)
-            .attr("stroke-opacity", dashboard.darktheme ? 0.35 : 0.25)
-            .attr("x1", this.margin.left)
-            .attr("x2", width - this.margin.right);
-
-        g.selectAll("circle")
-            .data(([, I]: any) => _d3.group(I, i => X[i]))
-            .enter()
-            .append("g")
-            .attr("id", ([, items]: any) => `circle${sessions[items[0]].id}`)
-            .attr("transform", ([, items]: any) => `translate(${xScale(X[items[0] as any])},0)`)
-            .append("g")
-            .attr("class", "circleScale")
-            .attr("id", ([, items]: any) => `circleScale${sessions[items[0]].id}`)
-            .style("cursor", "pointer")
-            .each(([, items]: any, i, nodes) => SessionCircle(nodes[i], items.map(i => sessions[i]), this.phaseSessionId, 5, !!this.phaseSessionId && items.some(i => sessions[i].id === this.phaseSessionId)));
+        errorBands.call(drawErrorBands);
+
+        const metaGroups = svg.append("g")
+            .attr("id", "meta-groups");
+
+        const drawMetaGroups = (groups: SelectionType) => {
+            const group = groups.selectAll(".meta")
+                .data(_d3.group(I, i => Y[i]), ([y]: any) => y)
+                .join(
+                    enter => enter.append("g")
+                        .attr("class", "meta")
+                        .attr("id", ([y]: any) => `meta${y as any}`)
+                        .attr("transform", ([y]: any) => `translate(0,${(yScale(y) as any) + 16})`)
+                        .each(([y]: any, i, nodes: any) => {
+                            const node = d3.select(nodes[i]);
+                            node.append("line")
+                                .attr("stroke", dashboard.darktheme ? "#6D6C6B" : "#4D4C4B")
+                                .attr("stroke-width", 1)
+                                .attr("stroke-linecap", 4)
+                                .attr("stroke-opacity", dashboard.darktheme ? 0.35 : 0.25)
+                                .attr("x1", this.margin.left)
+                                .attr("x2", width - this.margin.right);
+                        }),
+                    update => update.attr("transform", ([y]: any) => `translate(0,${(yScale(y) as any) + 16})`),
+                 );
+
+            group.selectAll(".session-circle")
+                .data(([, I]: any) => _d3.group(I, i => X[i]), ([, items]: any) => sessions[items[0]].id)
+                .join(
+                    enter => enter.append("g")
+                        .attr("class", "session-circle")
+                        .attr("id", ([, items]: any) => `circle${sessions[items[0]].id}`)
+                        .attr("transform", ([, items]: any) => `translate(${timelineHandler.timelineXScale!(X[items[0] as any])},0)`)
+                        .attr("visibility", ([, items]: any) => X[items[0] as any] >= timelineHandler.startDate! && X[items[0] as any] <= timelineHandler.endDate! ? "visible" : "hidden")
+                        .append("g")
+                        .attr("class", "circleScale")
+                        .attr("id", ([, items]: any) => `circleScale${sessions[items[0]].id}`)
+                        .style("cursor", "pointer")
+                        .each(([, items]: any, i, nodes: any) => SessionCircle(nodes[i], items.map(i => sessions[i]), this.phaseSessionId, circleSize, !!this.phaseSessionId && items.some(i => sessions[i].id === this.phaseSessionId))),
+                    update => update.attr("transform", ([, items]: any) => `translate(${timelineHandler.timelineXScale!(X[items[0] as any])},0)`)
+                        .attr("visibility", ([, items]: any) => X[items[0] as any] >= timelineHandler.startDate! && X[items[0] as any] <= timelineHandler.endDate! ? "visible" : "hidden"),
+                 );
+        }
+        metaGroups.call(drawMetaGroups);
 
         const reDrawCircles = (items: PhaseSessionResult[], index: number) => {
             const circle = svg.select(`#circleScale${items[0].id}`);
             circle.selectAll("*").remove();
-            SessionCircle(circle.node() as any, items, items[index].id, 5, !!this.phaseSessionId && items.some(i => i.id === this.phaseSessionId));
+            SessionCircle(circle.node() as any, items, items[index].id, circleSize, !!this.phaseSessionId && items.some(i => i.id === this.phaseSessionId));
         }
 
         // Attempt to select an existing div with id 'tooltip' from the body, or create one if it doesn't exist
@@ -259,7 +266,7 @@ export class PhaseHistoryGraph {
 
         const closestData = (x: number, y: number): PhaseSessionResult[] | undefined => {
 
-            if (x < this.margin.left - 16) {
+            if (x < this.margin.left - 16 || y < this.margin.top + 16) {
                 return undefined;
             }
 
@@ -268,7 +275,7 @@ export class PhaseHistoryGraph {
             const closest = sessions.reduce((best, session, i) => {
                 const absy = Math.abs(yScale(Y[0]) - y);
                 const sessionTime = handler.commitIdDates.get(session.commitId)!.getTime() / 1000;
-                const absx = Math.abs(xScale(sessionTime) - x);
+                const absx = Math.abs(timelineHandler.timelineXScale!(sessionTime) - x);
 
                 const lengthSqr = absy * absy + absx * absx;
 
@@ -330,7 +337,7 @@ export class PhaseHistoryGraph {
                 svg.select(`#circle${closest.id}`).raise();
 
                 const timeStamp = handler.commitIdDates.get(closest.commitId)!.getTime() / 1000;
-                const tx = xScale(timeStamp);
+                const tx = timelineHandler.timelineXScale!(timeStamp);
                 const ty = yScale(metaIndex)!;
                 const desc = getInfoSession(closestItems, index);
                 // Compute the bounding rectangle of the container to get its absolute position
@@ -339,7 +346,7 @@ export class PhaseHistoryGraph {
                 const offsetY = containerRect.top + window.scrollY;
                 const absTx = tx + offsetX;
                 const absTy = ty + offsetY;
-                this.updateTooltip(absTx, absTy, desc, closestItems.length > 1);
+                this.updateTooltip(absTx, absTy, desc);
             }
 
         }
@@ -349,7 +356,6 @@ export class PhaseHistoryGraph {
             if (this.tooltipItem) {
                 svg.selectAll(".circleScale").attr("transform", `scale(1)`);
                 tooltip.style("display", "none");
-                tooltip.classed("tooltip-block-scroll", false);
                 let index = this.tooltipItem.items.length > 1? this.tooltipItem.items.findIndex(item => item.id === this.phaseSessionId) : 0;
                 if (index < 0) index = 0;
                 reDrawCircles(this.tooltipItem.items, index);
@@ -389,8 +395,9 @@ export class PhaseHistoryGraph {
 
                     reDrawCircles(this.tooltipItem.items, this.tooltipItem.index);
                     const desc = getInfoSession(this.tooltipItem.items, this.tooltipItem.index);
-                    this.updateTooltip(undefined, undefined, desc, true);
+                    this.updateTooltip(undefined, undefined, desc);
                 }
+                event.preventDefault();
             }
             
         }
@@ -399,10 +406,9 @@ export class PhaseHistoryGraph {
         svg.on("mouseleave", (event) => handleMouseLeave(event));
         svg.on("click", (event) => handleMouseClick(event));
         svg.on("wheel", (event) => handleMouseWheel(event));
-
     }
 
-    updateTooltip(x?: number, y?: number, html?: string, blockWheel?: boolean) {
+    updateTooltip(x?: number, y?: number, html?: string) {
         if (!this.tooltip) {
             return;
         }
@@ -415,7 +421,6 @@ export class PhaseHistoryGraph {
 
         this.tooltip
             .style("display", "block")
-            .classed("tooltip-block-scroll", blockWheel ?? false)
             .html(html ?? "")
             .style("position", `absolute`)
             .style("width", `max-content`)
@@ -425,7 +430,7 @@ export class PhaseHistoryGraph {
             .style("line-height", "16px")
             .style("shapeRendering", "crispEdges")
             .style("stroke", "none")
-            .style("z-index", 10000)
+            .style("z-index", 2000000);
 
     }
 
@@ -433,6 +438,19 @@ export class PhaseHistoryGraph {
         d3.select(container).selectAll("*").remove();
     }
 
+    cleanupCallback() {
+        this.timelineCleanupCallback?.();
+        this.timelineCleanupCallback = undefined;
+    }
+
+    get averageFrequencyHours() {
+        if (this._averageSessionFrequencyHours !== undefined) {
+            return this._averageSessionFrequencyHours;
+        }
+        this._averageSessionFrequencyHours = calculateAverageSessionFrequencyHours(this.handler, this.sessions);
+        return this._averageSessionFrequencyHours;
+    }
+
     // refs
     test: TestNameRef;
     meta: MetadataRef;
@@ -440,9 +458,10 @@ export class PhaseHistoryGraph {
     phaseSessionId?: string;
 
     onClickCallback?: (session: PhaseSessionResult) => void;
+    timelineCleanupCallback?: () => void;
 
     handler: TestDataHandler;
-    margin: { top: number, right: number, bottom: number, left: number }
+    margin: Margin;
 
     hasRendered = false;
 
@@ -451,6 +470,8 @@ export class PhaseHistoryGraph {
 
     tooltip?: DivSelectionType;
     tooltipItem?: {index: number, items: PhaseSessionResult[], sessionId: string, delta: number};
+
+    private _averageSessionFrequencyHours?: number;
 }
 
 const SessionCircle = (container: SVGGElement, sessions: PhaseSessionResult[], sessionId: string | undefined, radius: number, selected?: boolean) => {
@@ -458,6 +479,7 @@ const SessionCircle = (container: SVGGElement, sessions: PhaseSessionResult[], s
     const color = getPhaseSessionStatusColor(session);
     const g = d3.select(container);
     g.append("circle")
+        .attr("stroke", "none")
         .attr("fill", color as any)
         .attr("r", radius);
 
diff --git a/Engine/Source/Programs/Horde/HordeDashboard/plugins/test-automation/testAutomationCommon.tsx b/Engine/Source/Programs/Horde/HordeDashboard/plugins/test-automation/testAutomationCommon.tsx
--- a/Engine/Source/Programs/Horde/HordeDashboard/plugins/test-automation/testAutomationCommon.tsx
+++ b/Engine/Source/Programs/Horde/HordeDashboard/plugins/test-automation/testAutomationCommon.tsx
@@ -463,6 +463,7 @@ export const PhaseHistoryGraphWidget: React.FC<{test: TestNameRef, phaseKey: str
     useEffect(() => {
         const graph = new PhaseHistoryGraph(test, phaseKey, streamId, handler, onClick);
         setState({ graph: graph });
+        return () => graph.cleanupCallback();
     }, [test, phaseKey, streamId, onClick]);
 
     useEffect(() => {
@@ -472,7 +473,7 @@ export const PhaseHistoryGraphWidget: React.FC<{test: TestNameRef, phaseKey: str
         } catch (err) {
             console.error(err);
         }
-    }, [container, state, meta, sessionId])
+    }, [container, state, meta, sessionId, withoutHeader, reverse])
 
     if (!state.graph) {
         return null;
diff --git a/Engine/Source/Programs/Horde/HordeDashboard/plugins/test-automation/testAutomationPhaseGridView.tsx b/Engine/Source/Programs/Horde/HordeDashboard/plugins/test-automation/testAutomationPhaseGridView.tsx
--- a/Engine/Source/Programs/Horde/HordeDashboard/plugins/test-automation/testAutomationPhaseGridView.tsx
+++ b/Engine/Source/Programs/Horde/HordeDashboard/plugins/test-automation/testAutomationPhaseGridView.tsx
@@ -924,7 +924,7 @@ export const TestPhasesGridView: React.FC<{ testId: string, handler: TestDataHan
                             }
                             {!!test && !handler.queryLoading && !handler.phaseQueryLoading && !handler.testDataQueryLoading &&
                                 <div ref={(ref) => { listBoundaryRef.current = ref; }} style={{ height: '100%' }}>
-                                    <ScrollablePane styles={{ contentContainer: {scrollbarGutter: 'stable', selectors: {'body:has(.tooltip-block-scroll) &': { overflowY: 'hidden !important' }} }}}>
+                                    <ScrollablePane>
                                         <DetailsList
                                             styles={{ root: {overflowX: 'hidden !important'}}}
                                             items={metaSelectorContext.filteredPhases}
diff --git a/Engine/Source/Programs/Horde/HordeDashboard/plugins/test-automation/testAutomationSummary.tsx b/Engine/Source/Programs/Horde/HordeDashboard/plugins/test-automation/testAutomationSummary.tsx
--- a/Engine/Source/Programs/Horde/HordeDashboard/plugins/test-automation/testAutomationSummary.tsx
+++ b/Engine/Source/Programs/Horde/HordeDashboard/plugins/test-automation/testAutomationSummary.tsx
@@ -23,6 +23,7 @@ export const TestHistoryGraphWidget: React.FC<{ test: TestNameRef, stream: strin
     useEffect(() => {
         const graph = new TestHistoryGraph(test, stream, handler, commonMetaKeys, onClick);
         setState({ graph: graph });
+        return () => graph.cleanupCallback();
     }, [commonMetaKeys]);
 
     useEffect(() => {
diff --git a/Engine/Source/Programs/Horde/HordeDashboard/plugins/test-automation/testHistoryGraph.ts b/Engine/Source/Programs/Horde/HordeDashboard/plugins/test-automation/testHistoryGraph.ts
--- a/Engine/Source/Programs/Horde/HordeDashboard/plugins/test-automation/testHistoryGraph.ts
+++ b/Engine/Source/Programs/Horde/HordeDashboard/plugins/test-automation/testHistoryGraph.ts
@@ -8,15 +8,20 @@ import { getHordeStyling } from 'horde/styles/Styles';
 import { getStatusColors, statusTexts, SessionValues, SessionListSelector } from "./testAutomationCommon";
 import { renderToString } from "react-dom/server"
 import { TestOutcome } from "./api";
+import { timelineHandler, Margin, constructTimeline, calculateAverageSessionFrequencyHours } from "./timelineGraph";
 
 // Handle bad "@types/d3" types, fix if addressed upstream
 const _d3 = d3 as any;
 
 type SelectionType = d3.Selection<SVGGElement, unknown, null, undefined>;
-type DivSelectionType = d3.Selection<HTMLDivElement, unknown, null, undefined>;
+type DivSelectionType = d3.Selection<HTMLDivElement, unknown, HTMLElement, undefined>;
 
 const clampMin = (x: number) => x !== 0 && x <= 0.3? 0.3: x;
 
+const width = 1000;
+const circleSize = 5;
+const highFrequencyThreshold = 4; // in hours
+
 export class TestHistoryGraph {
 
     constructor(test: TestNameRef, stream: string, handler: TestDataHandler, commonMetaKeys?: string[], onClickCallback?: (session: TestSessionResult) => void) {
@@ -66,19 +71,24 @@ export class TestHistoryGraph {
         const { modeColors } = getHordeStyling();
 
         this.clear(container);
+        this.cleanupCallback();
 
         this.hasRendered = true;
 
         this.initData(selectedMetaKeys);
 
+        // initialize timeline handle if not set
+        timelineHandler.init(this.test.id, this.handler, width, this.margin, false);
+        if (timelineHandler.isTimelineAtExtreme && this.averageFrequencyHours < highFrequencyThreshold) {
+            timelineHandler.rescaleTimeline(this.averageFrequencyHours / highFrequencyThreshold);
+        }
+
         const handler = this.handler;
         const sessions = this.sessions;
-        const width = 1000
 
         const X = _d3.map(sessions, (r) => handler.commitIdDates.get(r.commitId)!.getTime() / 1000);
         const Y = _d3.map(sessions, (r) => this.metaInfo.get(r.metadataId)!.index);
 
-        const xDomain = d3.extent(handler.commitIdDates.values() as any, (d:any) => d.getTime() / 1000);
         const yDomain = new _d3.InternSet(Y as any);
 
         const I = d3.range(X.length);
@@ -86,10 +96,8 @@ export class TestHistoryGraph {
         const yPadding = 1;
         const height = Math.ceil((yDomain.size + yPadding) * 16) + this.margin.top + this.margin.bottom;
 
-        const xRange = [this.margin.left, width - this.margin.right];
         const yRange = [this.margin.top, height - this.margin.bottom];
 
-        const xScale = _d3.scaleTime(xDomain, xRange);
         const yScale = _d3.scalePoint(yDomain, yRange).round(true).padding(yPadding);
 
         const svg = d3.select(container)
@@ -98,117 +106,97 @@ export class TestHistoryGraph {
             .attr("height", height + 24)
             .attr("viewBox", [0, 0, width, height + 24] as any)
 
-        // top axis
-        const xAxis = (g: SelectionType) => {
-
-            const dateMin = new Date(xDomain[0]! * 1000);
-            const dateMax = new Date(xDomain[1]! * 1000);
-
-            let ticks: number[] = [];
-            for (const date of d3.timeDays(dateMin, dateMax, 1).reverse()) {
-                ticks.push(date.getTime() / 1000);
-            }
-
-            if (ticks.length > 14) {
-                let nticks = [...ticks];
-                // remove first and last, will be readded 
-                const first = nticks.shift()!;
-                const last = nticks.pop()!;
-
-                const n = Math.floor(nticks.length / 12);
-
-                const rticks: number[] = [];
-                for (let i = 0; i < nticks.length; i = i + n) {
-                rticks.push(nticks[i]);
-                }
-
-                rticks.unshift(first);
-                rticks.push(last);
-                ticks = rticks;
-
-            }
-
-            g.attr("transform", `translate(0,16)`)
-                .style("font-family", "Horde Open Sans Regular")
-                .style("font-size", "9px")
-                .call(d3.axisTop(xScale)
-                .tickValues(ticks)
-                .tickFormat(d => getHumanTime(new Date((d as number) * 1000)))
-                .tickSizeOuter(0))
-                .call(g => g.select(".domain").remove())
-                .call(g => g.selectAll(".tick line")
-                    .attr("stroke-opacity", dashboard.darktheme ? 0.35 : 0.25).clone()
-                    .attr("stroke", dashboard.darktheme ? "#6D6C6B" : "#4D4C4B")
-                    .attr("y2", height - this.margin.bottom)
-                )
+        this.timelineCleanupCallback = constructTimeline(
+            svg,
+            () => {
+                // update meta groups
+                metaGroups.call(drawMetaGroups);
+            },
+            height - this.margin.bottom
+        );    
+
+        const metaGroups = svg.append("g")
+            .attr("id", "meta-groups");
+
+        const drawMetaGroups = (groups: SelectionType) => {
+            const group = groups.selectAll(".meta")
+                .data(_d3.group(I, i => Y[i]), ([y]: any) => y)
+                .join(
+                    enter => enter.append("g")
+                        .attr("class", "meta")
+                        .attr("id", ([y]: any) => `meta${y as any}`)
+                        .attr("transform", ([y]: any) => `translate(0,${(yScale(y) as any) + 16})`)
+                        .each(([y]: any, i, nodes: any) => {
+                            const node = d3.select(nodes[i]);
+                            node.append("line")
+                                .attr("stroke", dashboard.darktheme ? "#6D6C6B" : "#4D4C4B")
+                                .attr("stroke-width", 1)
+                                .attr("stroke-linecap", 4)
+                                .attr("stroke-opacity", dashboard.darktheme ? 0.35 : 0.25)
+                                .attr("x1", this.margin.left)
+                                .attr("x2", width - this.margin.right);
+
+                            node.append("text")
+                                .attr("text-anchor", "start")
+                                .style("alignment-baseline", "left")
+                                .style("font-family", "Horde Open Sans Regular")
+                                .style("font-size", 10)
+                                .attr("dy", "0.15em") // center stream name
+                                .attr("fill", dashboard.darktheme ? "#E0E0E0" : "#2D3F5F")
+                                .text(([y]: any) => this.metaNames.get(this.metaRefs[y])!);
+                        }),
+                    update => update.attr("transform", ([y]: any) => `translate(0,${(yScale(y) as any) + 16})`),
+                 );
+
+            group.selectAll(".session-circle")
+                .data(([, I]: any) => _d3.group(I, i => X[i]), ([, items]: any) => sessions[items[0]].id)
+                .join(
+                    enter => enter.append("g")
+                        .attr("class", "session-circle")
+                        .attr("id", ([, items]: any) => `circle${sessions[items[0]].id}`)
+                        .attr("transform", ([, items]: any) => `translate(${timelineHandler.timelineXScale!(X[items[0] as any])},0)`)
+                        .attr("visibility", ([, items]: any) => X[items[0] as any] >= timelineHandler.startDate! && X[items[0] as any] <= timelineHandler.endDate! ? "visible" : "hidden")
+                        .append("g")
+                        .attr("class", "circleScale")
+                        .attr("id", ([, items]: any) => `circleScale${sessions[items[0]].id}`)
+                        .style("cursor", "pointer")
+                        .each(([, items]: any, i, nodes: any) => SessionConcentricCircle(nodes[i], items.map(i => sessions[i]), 0, circleSize)),
+                    update => update.attr("transform", ([, items]: any) => `translate(${timelineHandler.timelineXScale!(X[items[0] as any])},0)`)
+                        .attr("visibility", ([, items]: any) => X[items[0] as any] >= timelineHandler.startDate! && X[items[0] as any] <= timelineHandler.endDate! ? "visible" : "hidden"),
+                 );
         }
-
-        svg.append("g").call(xAxis)
-
-        const g = svg.append("g")
-            .selectAll()
-            .data(_d3.group(I, i => Y[i]))
-            .join("g")
-            .attr("class", "meta")
-            .attr("id", ([y]: any) => `meta${y as any}`)
-            .attr("transform", ([y]: any) => `translate(0,${(yScale(y) as any) + 16})`);
-
-        g.append("line")
-            .attr("stroke", dashboard.darktheme ? "#6D6C6B" : "#4D4C4B")
-            .attr("stroke-width", 1)
-            .attr("stroke-linecap", 4)
-            .attr("stroke-opacity", dashboard.darktheme ? 0.35 : 0.25)
-            .attr("x1", this.margin.left)
-            .attr("x2", width - this.margin.right);
-
-        g.append("text")
-            .attr("text-anchor", "start")
-            .style("alignment-baseline", "left")
-            .style("font-family", "Horde Open Sans Regular")
-            .style("font-size", 10)
-            .attr("dy", "0.15em") // center stream name
-            .attr("x", () => 0)
-            .attr("fill", dashboard.darktheme ? "#E0E0E0" : "#2D3F5F")
-            .text(([y]: any) => this.metaNames.get(this.metaRefs[y])!);
-
-        g.selectAll("circle")
-            .data(([, I]: any) => _d3.group(I, i => X[i]))
-            .enter()
-            .append("g")
-            .attr("id", ([, items]: any) => `circle${sessions[items[0]].id}`)
-            .attr("transform", ([, items]: any) => `translate(${xScale(X[items[0] as any])},0)`)
-            .append("g")
-            .attr("class", "circleScale")
-            .attr("id", ([, items]: any) => `circleScale${sessions[items[0]].id}`)
-            .style("cursor", "pointer")
-            .each(([, items]: any, i, nodes) => SessionConcentricCircle(nodes[i], items.map(i => sessions[i]), 0, 5));
-
+        metaGroups.call(drawMetaGroups);
 
         const reDrawCircles = (items: TestSessionResult[], index: number) => {
             const circle = svg.select(`#circleScale${items[0].id}`);
             circle.selectAll("*").remove();
-            SessionConcentricCircle(circle.node() as any, items, index, 5);
+            SessionConcentricCircle(circle.node() as any, items, index, circleSize);
         }
 
-        const tooltip = this.tooltip = d3.select(container)
-            .append("div")
-            .attr("id", "tooltip")
-            .style("display", "none")
-            .style("background-color", modeColors.background)
-            .style("border", "solid")
-            .style("border-width", "1px")
-            .style("border-radius", "3px")
-            .style("border-color", dashboard.darktheme ? "#413F3D" : "#2D3F5F")
-            .style("padding", "6px")
-            .style("position", "absolute")
-            .style("pointer-events", "none");
+        // Attempt to select an existing div with id 'tooltip' from the body, or create one if it doesn't exist
+        let tooltip: DivSelectionType = d3.select("body").select("div#tooltip");
+        if (tooltip.empty()) {
+            tooltip = d3.select("body")
+                .append("div")
+                .attr("id", "tooltip")
+                .style("display", "none")
+                .style("background-color", modeColors.background)
+                .style("border", "solid")
+                .style("border-width", "1px")
+                .style("border-radius", "3px")
+                .style("border-color", dashboard.darktheme ? "#413F3D" : "#2D3F5F")
+                .style("padding", "6px")
+                .style("position", "absolute")
+                .style("pointer-events", "none");
+        }
+        this.tooltip = tooltip;
 
         // populate metaInfo y value
         this.metaInfo.forEach(meta => meta.y = yScale(meta.index as any)!);
 
         const closestData = (x: number, y: number): TestSessionResult[] | undefined => {
 
-            if (x < this.margin.left - 16) {
+            if (x < this.margin.left - 16 || y < this.margin.top + 16) {
                 return undefined;
             }
 
@@ -217,7 +205,7 @@ export class TestHistoryGraph {
             const closest = sessions.reduce((best, session, i) => {
                 const absy = Math.abs(this.metaInfo.get(session.metadataId)!.y - y);
                 const sessionTime = handler.commitIdDates.get(session.commitId)!.getTime() / 1000;
-                const absx = Math.abs(xScale(sessionTime) - x);
+                const absx = Math.abs(timelineHandler.timelineXScale!(sessionTime) - x);
 
                 const lengthSqr = absy * absy + absx * absx;
 
@@ -254,7 +242,7 @@ export class TestHistoryGraph {
             if (item.phasesSucceededCount + item.phasesFailedCount + item.phasesUnspecifiedCount > 1)
                 desc += renderToString(SessionValues(item, {fontSize: 9}));
 
-            return desc;
+            return `<span style="color: ${modeColors.text};">${desc}</span>`;;
         }
 
         const handleMouseMove = (event: any) => {
@@ -280,11 +268,17 @@ export class TestHistoryGraph {
                 svg.select(`#meta${metaIndex}`).raise();
 
                 const timeStamp = handler.commitIdDates.get(closest.commitId)!.getTime() / 1000;
-                const tx = xScale(timeStamp);
+                const tx = timelineHandler.timelineXScale!(timeStamp);
                 const ty = yScale(metaIndex as any)!
                 const desc = getInfoSession(closestItems, 0);
-
-                this.updateTooltip(tx, ty, desc);
+                // Compute the bounding rectangle of the container to get its absolute position
+                const containerRect = container.getBoundingClientRect();
+                const offsetX = containerRect.left + window.scrollX;
+                const offsetY = containerRect.top + window.scrollY;
+                const absTx = tx + offsetX;
+                const absTy = ty + offsetY;
+
+                this.updateTooltip(absTx, absTy, desc);
             }
 
         }
@@ -309,6 +303,7 @@ export class TestHistoryGraph {
                 const closestItems = closestData(mouseX, mouseY);
                 if (closestItems && this.onClickCallback) {
                     this.onClickCallback(closestItems[this.tooltipItem.index]);
+                    handleMouseLeave(undefined);
                 }
             }
 
@@ -332,6 +327,7 @@ export class TestHistoryGraph {
                     const desc = getInfoSession(this.tooltipItem.items, this.tooltipItem.index);
                     this.updateTooltip(undefined, undefined, desc);
                 }
+                event.preventDefault();
             }
             
         }
@@ -364,7 +360,8 @@ export class TestHistoryGraph {
             .style("font-size", "10px")
             .style("line-height", "16px")
             .style("shapeRendering", "crispEdges")
-            .style("stroke", "none");
+            .style("stroke", "none")
+            .style("z-index", 2000000);
 
     }
 
@@ -372,14 +369,28 @@ export class TestHistoryGraph {
         d3.select(container).selectAll("*").remove();
     }
 
+    cleanupCallback() {
+        this.timelineCleanupCallback?.();
+        this.timelineCleanupCallback = undefined;
+    }
+
+    get averageFrequencyHours() {
+        if (this._averageSessionFrequencyHours !== undefined) {
+            return this._averageSessionFrequencyHours;
+        }
+        this._averageSessionFrequencyHours = calculateAverageSessionFrequencyHours(this.handler, this.sessions);
+        return this._averageSessionFrequencyHours;
+    }
+
     // test ref
     test: TestNameRef;
 
     commonMetaKeys?: string[];
     onClickCallback?: (session: TestSessionResult) => void;
+    timelineCleanupCallback?: () => void;
 
     handler: TestDataHandler;
-    margin: { top: number, right: number, bottom: number, left: number }
+    margin: Margin;
 
     hasRendered = false;
 
@@ -394,6 +405,8 @@ export class TestHistoryGraph {
 
     tooltip?: DivSelectionType;
     tooltipItem?: {index: number, items: TestSessionResult[], sessionId: string, delta: number};
+
+    private _averageSessionFrequencyHours?: number;
 }
 
 type ItemStack = {value: number, name: string, color: string, title?: boolean};
diff --git a/Engine/Source/Programs/Horde/HordeDashboard/plugins/test-automation/timelineGraph.ts b/Engine/Source/Programs/Horde/HordeDashboard/plugins/test-automation/timelineGraph.ts
new file mode 100644
--- /dev/null
+++ b/Engine/Source/Programs/Horde/HordeDashboard/plugins/test-automation/timelineGraph.ts
@@ -0,0 +1,450 @@
+// Copyright Epic Games, Inc. All Rights Reserved.
+
+import * as d3 from "d3";
+import { BaseType } from "d3";
+import { PhaseSessionResult, TestDataHandler, TestSessionResult } from "./testData";
+import { getHordeStyling } from "horde/styles/Styles";
+import { getHumanTime } from "horde/base/utilities/timeUtils";
+import dashboard from "horde/backend/Dashboard";
+
+// Handle bad "@types/d3" types, fix if addressed upstream
+const _d3 = d3 as any;
+
+type SelectionType = d3.Selection<SVGGElement, unknown, null, undefined>;
+type SelectionBaseType = d3.Selection<BaseType, unknown, SVGGElement, unknown>;
+export type Margin = { top: number, right: number, bottom: number, left: number }
+
+class TimelineHandler {
+    constructor() {
+        this._cachedTimelines = new Map();
+    }
+
+    init(testid: string, handler: TestDataHandler, width: number, margin: Margin, reverse?: boolean) {
+        if (this._testid !== testid)
+        {
+            this._testid = testid;
+            if (!this.minDate || !this.maxDate) {
+                const dates = handler.commitIdDates.values().map(d => d.getTime() / 1000).toArray();
+                this.minDate = Math.min(...dates);
+                this.maxDate = Math.max(...dates);
+                this.ticks = this.getTicks();
+                this.axisXScale = _d3.scaleTime().domain([this.minDate, this.maxDate]);
+            }
+            if (!this._cachedTimelines.has(testid)) {
+                this._cachedTimelines.set(testid, {startDate: this.minDate!, endDate: this.maxDate!});
+            }
+            this.timelineXScale = _d3.scaleTime();
+            this.callbacks = new Set();
+        }
+
+        this.axisXScale!.range(reverse? [width - margin.right, margin.left] : [margin.left, width - margin.right]);
+        this.timelineXScale!.range(this.axisXScale!.range());
+        this.axisRangeWidth = Math.abs(this.axisXScale!.range()[1] - this.axisXScale!.range()[0]);
+        this.updateTimelineXScale();
+    }
+
+    updateTimelineXScale() {
+        const start = Math.min(this.startDate, this.endDate);
+        const end = Math.max(this.startDate, this.endDate);
+        this.startDate = start;
+        this.endDate = end;
+        this.timelineXScale.domain([start, end]);
+    }
+
+    updateDate(x: number, isStart: boolean) {
+        const date = this.axisXScale!.invert(x);
+        if (isStart) {
+            this.startDate = date;
+        } else {
+            this.endDate = date;
+        }
+        this.updateTimelineXScale();
+    }
+
+    getXStart() {
+        return this.axisXScale!(this.startDate!);
+    }
+    getXEnd() {
+        return this.axisXScale!(this.endDate!);
+    }
+
+    getTimelineRangeWidth() {
+        return Math.abs(this.getXStart() - this.getXEnd());
+    }
+
+    addCallback(callback: () => void) {
+        this.callbacks!.add(callback);
+    }
+    removeCallback(callback: () => void) {
+        this.callbacks!.delete(callback);
+    }
+    triggerCallbacks() {
+        // Trigger callbacks by chunks using promises with cancellation.
+        // This is to avoid blocking the main thread for too long.
+        if (this._triggerCallbacksAbortController) {
+            this._triggerCallbacksAbortController.abort();
+        }
+        this._triggerCallbacksAbortController = new AbortController();
+        const localAbort = this._triggerCallbacksAbortController;
+        const callbacks = Array.from(this.callbacks!);
+        const MAX_DELTA_MS = 10; // Target max milliseconds processing before yielding
+        const defer = () => new Promise(resolve => setTimeout(resolve, 0));
+
+        const run = async () => {
+            let i = 0;
+            while (i < callbacks.length) {
+                const startTime = performance.now();
+                for (; i < callbacks.length; ++i) {
+                    try { callbacks[i](); } catch (err) { /* ignore individual callback errors */ }
+                    const now = performance.now();
+                    if (now - startTime > MAX_DELTA_MS) {
+                        break;
+                    }
+                }
+                if (i < callbacks.length) {
+                    await defer();
+                    if (localAbort.signal.aborted) break;
+                }
+            }
+        };
+
+        run().finally(() => {
+            if (this._triggerCallbacksAbortController === localAbort) {
+                // Only clear if this is still most recent
+                this._triggerCallbacksAbortController = undefined;
+            }
+        });
+    }
+
+    get isTimelineAtExtreme() {
+        return this.endDate === this.maxDate && this.startDate === this.minDate;
+    }
+
+    rescaleTimeline(factor: number) {
+        if (!this.isTimelineAtExtreme) {
+            // if timeline is not at the extremes, don't rescale
+            return;
+        }
+        this.startDate = this.endDate! - ((this.endDate! - this.startDate!) * factor);
+        this.updateTimelineXScale();
+        this.triggerCallbacks(); // trigger callbacks to update previously rendered timeline
+    }
+
+    private _triggerCallbacksAbortController?: AbortController;
+
+    private getTicks() {
+        const dateMin = new Date(this.minDate! * 1000);
+        const dateMax = new Date(this.maxDate! * 1000);
+
+        let ticks: number[] = [];
+        for (const date of d3.timeDays(dateMin, dateMax, 1).reverse()) {
+            ticks.push(date.getTime() / 1000);
+        }
+
+        if (ticks.length > 14) {
+            let nticks = [...ticks];
+            // remove first and last, will be readded 
+            const first = nticks.shift()!;
+            const last = nticks.pop()!;
+
+            const n = Math.floor(nticks.length / 12);
+
+            const rticks: number[] = [];
+            for (let i = 0; i < nticks.length; i = i + n) {
+                rticks.push(nticks[i]);
+            }
+
+            rticks.unshift(first);
+            rticks.push(last);
+            ticks = rticks;
+        }
+
+        return ticks;
+    }
+
+    get startDate() {
+        return this._cachedTimelines.get(this._testid!)!.startDate;
+    }
+    get endDate() {
+        return this._cachedTimelines.get(this._testid!)!.endDate;
+    }
+    set startDate(date: number) {
+        this._cachedTimelines.get(this._testid!)!.startDate = date;
+    }
+    set endDate(date: number) {
+        this._cachedTimelines.get(this._testid!)!.endDate = date;
+    }
+
+    private _testid?: string;
+    private _cachedTimelines: Map<string, {startDate: number, endDate: number}>;
+    maxDate?: number;
+    minDate?: number;
+    axisXScale?: any;
+    axisRangeWidth?: number;
+    timelineXScale?: any;
+    ticks?: number[];
+    callbacks?: Set<() => void>;
+}
+
+export const timelineHandler = new TimelineHandler();
+
+export const constructTimeline = (svg: SelectionType, onTimelineChange: () => void, height: number, reverse?: boolean, withoutHeader?: boolean) => {
+    const { modeColors } = getHordeStyling();
+
+    // top axis
+    const headerAxis = (element: SelectionBaseType) => {
+        element.selectAll(".tick")
+            .data(timelineHandler.ticks!)
+            .join(
+                enter => enter.append("g")
+                    .attr("class", "tick")
+                    .attr("transform", (d: number) => `translate(${timelineHandler.axisXScale!(d)},0)`)
+                    .each((d, i, nodes) => {
+                        if (!withoutHeader) {
+                            d3.select(nodes[i]).append("text")
+                                .attr("y", -8)
+                                .attr("class", "text")
+                                .attr("fill", modeColors.text)
+                                .attr("text-anchor", "middle")
+                                .text(d => getHumanTime(new Date((d as number) * 1000)))
+                        }
+                    })
+                    .append("line")
+                        .attr("class", "line")
+                        .attr("y1", -5)
+                        .attr("stroke-width", 1)
+                        .attr("stroke-linecap", "round")
+                        .attr("stroke", dashboard.darktheme ? "#6D6C6B" : "#4D4C4B"),
+                update => update.attr("transform", (d: number) => `translate(${timelineHandler.axisXScale!(d)},0)`),
+                exit => exit.remove()
+            );
+    }
+    const seriesAxis = (element: SelectionBaseType) => {
+        element.selectAll(".tick")
+            .data(timelineHandler.ticks!)
+            .join(
+                enter => enter.append("g")
+                    .attr("class", "tick")
+                    .attr("transform", (d: number) => `translate(${timelineHandler.timelineXScale!(d)},0)`)
+                    .attr("visibility", (d: number) => d >= timelineHandler.startDate! && d <= timelineHandler.endDate! ? "visible" : "hidden")
+                    .append("line")
+                        .attr("class", "line")
+                        .attr("y1", -3)
+                        .attr("stroke-width", 1)
+                        .attr("stroke-linecap", "round")
+                        .attr("stroke", dashboard.darktheme ? "#6D6C6B" : "#4D4C4B")
+                        .attr("stroke-opacity", dashboard.darktheme ? 0.35 : 0.25).clone()
+                        .attr("y2", height),
+                update => update.attr("transform", (d: number) => `translate(${timelineHandler.timelineXScale!(d)},0)`)
+                    .attr("visibility", (d: number) => d >= timelineHandler.startDate! && d <= timelineHandler.endDate! ? "visible" : "hidden"),
+                exit => exit.remove()
+            );
+    }
+    const xAxis = (g: SelectionType, header: boolean, series: boolean) => {
+        if (header) {
+            g.selectAll("#header-axis").call(headerAxis);
+        }
+        if (series) {
+            g.selectAll("#series-axis").call(seriesAxis);
+        }
+    }
+    const axis = svg.append("g");
+    axis.append("g")
+            .attr("id", "header-axis")
+            .attr("transform", 'translate(0,16)')
+            .style("font-family", "Horde Open Sans Regular")
+            .style("font-size", "9px");
+    axis.append("g")
+            .attr("id", "series-axis")
+            .attr("transform", 'translate(0,16)');
+    axis.call(xAxis, !withoutHeader, true);
+
+    const innerOnTimelineChange = () => {
+        // update axis
+        withoutHeader && axis.call(xAxis, false, true);
+        // trigger onTimelineChange callback
+        onTimelineChange();
+    }
+    timelineHandler.addCallback(innerOnTimelineChange);
+    const registeredCallbacks = [innerOnTimelineChange];
+
+    if (!withoutHeader) {
+        // dragable timeline
+        const dragHandle = (handle: SelectionType, updateHandle: (value: number, width?: number) => void, isAreaSelected?: boolean) => {
+            const [axisRangeStart, axisRangeEnd] = reverse? timelineHandler.axisXScale!.range().reverse() : timelineHandler.axisXScale!.range();
+            let startOffset = 0;
+            let width = 0;
+            return d3.drag()
+            .on("start", (e: any) => {
+                handle.attr("cursor", "grabbing");
+                startOffset = isAreaSelected ? e.x - Math.min(timelineHandler.getXStart(), timelineHandler.getXEnd()) : 0;
+                width = isAreaSelected ? timelineHandler.getTimelineRangeWidth() : 0;
+            })
+            .on("end", (e: any) => { isAreaSelected ? handle.call(checkHandleGrabCursor, width) : handle.attr("cursor", "grab") })
+            .on("drag", (e: any) => {
+                const x = Math.min(Math.max(e.x - startOffset, axisRangeStart), axisRangeEnd - width);
+                handle.attr("transform", `translate(${x},0)`);
+                // update date
+                updateHandle(x, width);
+                // update area selected
+                !isAreaSelected && areaSelected.call(updateAreaSelected);
+                // update axis
+                axis.call(xAxis, false, true);
+                // trigger callbacks
+                updatingThisHeader = true;
+                timelineHandler.triggerCallbacks();
+                updatingThisHeader = false;
+            });
+        }
+        const checkHandleGrabCursor = (handle: SelectionType, width: number) => {
+            handle.attr("cursor", Math.round(timelineHandler.axisRangeWidth! - width) > 20 ? "grab" : "default");
+        }
+
+        let updatingThisHeader = false;
+        const onHeaderTimelineChange = () => {
+            if (!updatingThisHeader) {
+                areaSelected.call(updateAreaSelected);
+                updateHandleDates();
+                // update axis
+                axis.call(xAxis, false, true);
+            }
+        }
+        timelineHandler.addCallback(onHeaderTimelineChange);
+        registeredCallbacks.push(onHeaderTimelineChange);
+
+        const areaSelected = svg.append("g");
+        areaSelected.append("rect")
+            .attr("id", "area-selected")
+            .attr("fill", modeColors.text)
+            .attr("opacity", 0.1)
+            .attr("x", 3)
+            .attr("y", 0)
+            .attr("height", 16);
+        const updateAreaSelected = (selection: SelectionType) => {
+            selection.attr("transform", `translate(${Math.min(timelineHandler.getXStart(), timelineHandler.getXEnd())},0)`);
+            selection.select("#area-selected").attr("width", Math.max(0, timelineHandler.getTimelineRangeWidth() - 6));
+        }
+        const updateHandleDates = () => {
+            handleADate = timelineHandler.startDate!;
+            handleA.attr("transform", `translate(${timelineHandler.getXStart()},0)`);
+            handleBDate = timelineHandler.endDate!;
+            handleB.attr("transform", `translate(${timelineHandler.getXEnd()},0)`);
+        }
+        areaSelected.call(updateAreaSelected);
+        areaSelected.call(
+            dragHandle(areaSelected, (v, width) => {
+                timelineHandler.updateDate(v, !reverse);
+                width !== undefined && timelineHandler.updateDate(v + width, !!reverse);
+                updateHandleDates();
+            }, true) as any
+        );
+        areaSelected.on("mouseover", (event) => {
+            const width = timelineHandler.getTimelineRangeWidth();
+            areaSelected.call(checkHandleGrabCursor, width);
+        });
+        areaSelected.on("wheel", (event: any) => {
+            event.preventDefault();
+            // Get the relative x position from the event target
+            const rect = event.target.getBoundingClientRect();
+            const relX = event.clientX - rect.left;
+            const baseWidth = timelineHandler.getTimelineRangeWidth();
+            const offsetFactor = relX / baseWidth;
+            const scaleFactor = event.deltaY < 0 ? 1.08 : 0.92;
+            let width = baseWidth * scaleFactor;
+            if (width > timelineHandler.axisRangeWidth!) {
+                width = timelineHandler.axisRangeWidth!;
+            } else if (width < 20) {
+                width = 20;
+            }
+            let xStart = Math.min(timelineHandler.getXStart()!, timelineHandler.getXEnd()!) + relX - offsetFactor * width;
+            let xEnd = xStart + width;
+            const [axisRangeStart, axisRangeEnd] = reverse? timelineHandler.axisXScale!.range().reverse() : timelineHandler.axisXScale!.range();
+            if (xEnd > axisRangeEnd) {
+                xEnd = axisRangeEnd;
+                xStart = xEnd - width;
+            }
+            if (xStart < axisRangeStart) {
+                xStart = axisRangeStart;
+                xEnd = xStart + width;
+            }
+            timelineHandler.updateDate(xStart, !reverse);
+            timelineHandler.updateDate(xEnd, !!reverse);
+            // update handles
+            areaSelected.call(updateAreaSelected);
+            areaSelected.call(checkHandleGrabCursor, width);
+            updateHandleDates();
+            // update axis
+            axis.call(xAxis, false, true);
+            // trigger callbacks
+            updatingThisHeader = true;
+            timelineHandler.triggerCallbacks();
+            updatingThisHeader = false;
+        });
+
+        const constructHandle = (handle) => {
+            handle.append("rect")
+                .attr("x", -3)
+                .attr("y", 0)
+                .attr("width", 6)
+                .attr("height", 16)
+                .attr("fill", `${modeColors.text}3f`);
+        }
+    
+        const handleA = svg.append("g")
+            .attr("cursor", "grab")
+            .attr("transform", `translate(${timelineHandler.getXStart()},0)`);
+        let handleADate = timelineHandler.startDate!;
+        handleA.call(
+                dragHandle(handleA, (v) => {
+                timelineHandler.updateDate(v, !Math.round(handleADate - timelineHandler.startDate!));
+                handleADate = timelineHandler.axisXScale!.invert(v);
+            }) as any
+            );
+        constructHandle(handleA);
+    
+        const handleB = svg.append("g")
+            .attr("cursor", "grab")
+            .attr("transform", `translate(${timelineHandler.getXEnd()},0)`);
+        let handleBDate = timelineHandler.endDate!;
+        handleB.call(
+                dragHandle(handleB, (v) => {
+                timelineHandler.updateDate(v, !Math.round(handleBDate - timelineHandler.startDate!));
+                handleBDate = timelineHandler.axisXScale!.invert(v);
+            }) as any
+            );
+        constructHandle(handleB);       
+    }
+
+    return () => {
+        for (const callback of registeredCallbacks) {
+            timelineHandler.removeCallback(callback);
+        }
+    }
+}
+
+export const calculateAverageSessionFrequencyHours = (handler: TestDataHandler, sessions: TestSessionResult[] | PhaseSessionResult[]) => {
+    let averageSessionFrequencyHours: number = Infinity;
+
+    const thresholdNumberOfSessions = 48; // Consider only the last 48 sessions for the average frequency
+    if (sessions.length > thresholdNumberOfSessions) {
+        // Filter sessions to only unique commitIds (preserve order)
+        const seenCommitIds = new Set<string>();
+        const uniqueDates: number[] = [];
+        for (let i = 0; i < Math.min(sessions.length, thresholdNumberOfSessions); i++) {
+            const commitId = sessions[i].commitId;
+            if (!seenCommitIds.has(commitId)) {
+                seenCommitIds.add(commitId);
+                uniqueDates.push(handler.commitIdDates.get(commitId)!.getTime() / 1000);
+            }
+        }
+
+        if (uniqueDates.length > 1) {
+            const firstDate = uniqueDates[0];
+            const lastDate = uniqueDates[uniqueDates.length - 1];
+            const totalSeconds = Math.abs(lastDate - firstDate);
+            averageSessionFrequencyHours = totalSeconds / ((uniqueDates.length - 1) * 3600);
+        }
+    }
+
+    return averageSessionFrequencyHours;
+}
