Horde: Update release notes

#rnx

diff --git a/Engine/Source/Programs/Horde/Docs/ReleaseNotes.md b/Engine/Source/Programs/Horde/Docs/ReleaseNotes.md
--- a/Engine/Source/Programs/Horde/Docs/ReleaseNotes.md
+++ b/Engine/Source/Programs/Horde/Docs/ReleaseNotes.md
@@ -1,5 +1,55 @@
 # Release Notes
 
+## 2026-03-23
+
+* EpicGames.OIDC: Prevent refresh token loss when process shuts down during token rotation (51945923)
+* Horde Dashboard - Add "Open in Build Storage Tool" button for Cloud Artifacts (51945444)
+* Pass job ID to label hover card ChangeButton for context menu
+  ChangeButton in LabelHoverCard was missing job ID, causing Open Job
+  Details and Copy Job links to point to /job/undefined. Pass jobId for
+  current build and prev.jobId for previous build CL buttons. (51928097)
+* [Horde Dashboard] Ensure unnamed pools aren't displayed in pool view filter results (51922279)
+* Fix JobView label layout when no label columns defined (51919471)
+* Label hover card with server-side details endpoint
+  Add GET /api/v1/jobs/{jobId}/labels/{labelIndex}/details endpoint
+  returning label state, failing steps with log events, build health
+  issues (with Slack/Jira links), previous build status, and
+  cross-stream merge flow in a single response.
+  New shared LabelBadge and LabelHoverCard frontend components used
+  across JobView, JobViewIncremental, JobViewAll, and JobDetailLabels.
+  Compact badge sizing for non-failure labels, hover row highlights,
+  inline issue modal on click, and configurable external issue URL. (51915812)
+* Ignore file actions of type "unknown" when unshelving changes in PerforceMaterializer (51892769)
+* Add template filter to job search (51889166)
+* Sending step messages to databricks. (51864061)
+* Disable expensive C# analyzers
+  CA5390 and CA5403 are very expensive analyzers to run under some conditions, and provide very little value.
+  This is especially noticable when building UGS with analyzers. Build time for that single .csproj went from 117s -> 19s when disabled.
+  A full Horde.sln build 160s -> 92s. (51820071)
+* Remove unused using directive in JobsController (51789024)
+* Add optimized label details endpoint for hover card
+  New GET /api/v1/jobs/{jobId}/labels/{labelIndex}/details endpoint that
+  aggregates steps, log events, issues, previous build, and cross-stream
+  merge chain in a single response. Replaces N+1 per-hover API calls. (51785105)
+* Move Robomerge config from server config to global config for dynamic reload via globals.json (51782440)
+* Add RoboMerge merge graph service to Build plugin
+  Parses branchmap JSON files from Perforce and serves merge topology
+  via REST API. Includes ticking service with change detection, merge
+  chain computation, reachability index, and stream-level authorization. (51771667)
+* EpicGames.OIDC: Use new name for token state file
+  Recent bug fixes to this library fixed a race condition where long running processes would save stale token data. By changing the name of this file, any apps using the outdated version of this library can no longer interfere. This adds an initial browser login to authenticate but after that, new and old can better co-exist. MacOS provider is not affected since it addresses token data directly by provider name (no saving/loading of file with mixed provider state). (51761382)
+* Simplify session conflict handling to use debug logging + batched Slack report
+  Replace the unused auto-quarantine implementation with a lightweight Redis sorted set: the first 20 mismatches per agent log at error level subsequent ones log at debug. (51708749)
+* Cloud Storage - Implemented the wrapper for listing file contents using ICloudStorage (51674981)
+* Swap registration order of exception logger vs handler
+  Surfaces the error messages in ExceptionLoggingMiddleware, which previously got silenced. This so we can control the format of "Unhandled exception" log lines. Once verified, we'll disable the generic ASP.NET based logger. (51674205)
+* Add more generic xcode parsing to find pre-compile errors arising from the .xcodeproj files. (51673579)
+* Always register BisectService for DI even when bisection is feature-flagged off
+  BisectService already checks DisableBisection at runtime in TickAsync,
+  so it is safe to register unconditionally. This ensures the controller
+  and related services remain available for dependency injection. (51669557)
+* EpicGames.OIDC: Use SingleInstanceMutex to avoid cross-thread lock release (51668299)
+
 ## 2026-03-12
 
 * Feature flag bisection - disabled by default (51640744)
