Horde: Fix ComputeMergeChain broken by side branches with forceFlowTo to Main
When multiple branches target the same branch via forceFlowTo
(e.g. both Dev and Dev-Authoring target Main),
the reverseMap silently overwrites, breaking the release spine.
Fix by collecting all candidates and preferring spine over leaf.

diff --git a/Engine/Source/Programs/Horde/Plugins/Build/HordeServer.Build.Tests/RoboMerge/RoboMergeBranchmapParserTests.cs b/Engine/Source/Programs/Horde/Plugins/Build/HordeServer.Build.Tests/RoboMerge/RoboMergeBranchmapParserTests.cs
--- a/Engine/Source/Programs/Horde/Plugins/Build/HordeServer.Build.Tests/RoboMerge/RoboMergeBranchmapParserTests.cs
+++ b/Engine/Source/Programs/Horde/Plugins/Build/HordeServer.Build.Tests/RoboMerge/RoboMergeBranchmapParserTests.cs
@@ -525,6 +525,59 @@ namespace HordeServer.Tests.RoboMerge
 			Assert.AreEqual("Main", chain.Entries[6].Branch.Name);
 		}
 
+		[TestMethod]
+		public void ComputeMergeChain_SideBranchToMainDoesNotBreakSpine()
+		{
+			// Spine: Release-1 → Release-2 → Dev → Main
+			// Side branch: Dev-Authoring also has forceFlowTo: ["Main"]
+			// The side branch should NOT clobber Dev in the reverse map
+			string json = @"{
+				""defaultStreamDepot"": ""Test"",
+				""branches"": [
+					{ ""name"": ""Main"" },
+					{ ""name"": ""Dev"", ""forceFlowTo"": [""Main""] },
+					{ ""name"": ""Release-2"", ""forceFlowTo"": [""Dev""] },
+					{ ""name"": ""Release-1"", ""forceFlowTo"": [""Release-2""] },
+					{ ""name"": ""Dev-Authoring"", ""forceFlowTo"": [""Main""] }
+				]
+			}";
+
+			MergeGraph graph = Parse("test.branchmap.json", json, 1);
+			MergeChain? chain = RoboMergeBranchmapParser.ComputeMergeChain(graph);
+
+			Assert.IsNotNull(chain);
+			Assert.AreEqual(4, chain!.Entries.Count);
+			Assert.AreEqual("Release-1", chain.Entries[0].Branch.Name);
+			Assert.AreEqual("Release-2", chain.Entries[1].Branch.Name);
+			Assert.AreEqual("Dev", chain.Entries[2].Branch.Name);
+			Assert.AreEqual("Main", chain.Entries[3].Branch.Name);
+		}
+
+		[TestMethod]
+		public void ComputeMergeChain_MultipleSideBranches()
+		{
+			// Multiple side branches targeting Main shouldn't break the spine
+			string json = @"{
+				""defaultStreamDepot"": ""Test"",
+				""branches"": [
+					{ ""name"": ""Main"" },
+					{ ""name"": ""Dev"", ""forceFlowTo"": [""Main""] },
+					{ ""name"": ""Release"", ""forceFlowTo"": [""Dev""] },
+					{ ""name"": ""Side-A"", ""forceFlowTo"": [""Main""] },
+					{ ""name"": ""Side-B"", ""forceFlowTo"": [""Main""] }
+				]
+			}";
+
+			MergeGraph graph = Parse("test.branchmap.json", json, 1);
+			MergeChain? chain = RoboMergeBranchmapParser.ComputeMergeChain(graph);
+
+			Assert.IsNotNull(chain);
+			Assert.AreEqual(3, chain!.Entries.Count);
+			Assert.AreEqual("Release", chain.Entries[0].Branch.Name);
+			Assert.AreEqual("Dev", chain.Entries[1].Branch.Name);
+			Assert.AreEqual("Main", chain.Entries[2].Branch.Name);
+		}
+
 		static MergeGraph Parse(string depotPath, string json, int headChange)
 		{
 			byte[] data = Encoding.UTF8.GetBytes(json);
diff --git a/Engine/Source/Programs/Horde/Plugins/Build/HordeServer.Build/RoboMerge/RoboMergeBranchmapParser.cs b/Engine/Source/Programs/Horde/Plugins/Build/HordeServer.Build/RoboMerge/RoboMergeBranchmapParser.cs
--- a/Engine/Source/Programs/Horde/Plugins/Build/HordeServer.Build/RoboMerge/RoboMergeBranchmapParser.cs
+++ b/Engine/Source/Programs/Horde/Plugins/Build/HordeServer.Build/RoboMerge/RoboMergeBranchmapParser.cs
@@ -150,8 +150,10 @@ namespace HordeServer.RoboMerge
 				}
 			}
 
-			// Build reverse forceFlowTo map: if B.ForceFlowTo contains T, then reverseMap[T] = B
-			Dictionary<string, string> reverseMap = new Dictionary<string, string>(StringComparer.OrdinalIgnoreCase);
+			// Build reverse forceFlowTo map: if B.ForceFlowTo contains T, then reverseMapMulti[T] includes B.
+			// Multiple branches may target the same branch (e.g. both Dev and Dev-Authoring target Main),
+			// so we collect all candidates and resolve conflicts below.
+			Dictionary<string, List<string>> reverseMapMulti = new Dictionary<string, List<string>>(StringComparer.OrdinalIgnoreCase);
 			HashSet<string> hasForceFlowTo = new HashSet<string>(StringComparer.OrdinalIgnoreCase);
 
 			foreach (MergeBranch branch in enabledBranches.Values)
@@ -171,16 +173,51 @@ namespace HordeServer.RoboMerge
 								graph.BotName, branch.Name, target);
 							continue;
 						}
-						reverseMap[target] = branch.Name;
+						if (!reverseMapMulti.TryGetValue(target, out List<string>? sources))
+						{
+							sources = new List<string>();
+							reverseMapMulti[target] = sources;
+						}
+						sources.Add(branch.Name);
 					}
 				}
 			}
 
-			if (reverseMap.Count == 0)
+			if (reverseMapMulti.Count == 0)
 			{
 				return null;
 			}
 
+			// Resolve conflicts: when multiple branches target the same branch, prefer the one that is itself
+			// a forceFlowTo target (i.e. part of the spine, not a side branch). For example, if both Dev
+			// and Dev-Authoring target Main, prefer Dev because Release targets it.
+			HashSet<string> isTargeted = new HashSet<string>(reverseMapMulti.Keys, StringComparer.OrdinalIgnoreCase);
+			Dictionary<string, string> reverseMap = new Dictionary<string, string>(StringComparer.OrdinalIgnoreCase);
+
+			foreach ((string target, List<string> sources) in reverseMapMulti)
+			{
+				if (sources.Count == 1)
+				{
+					reverseMap[target] = sources[0];
+				}
+				else
+				{
+					// Prefer the source that is itself a target (spine branch over leaf branch)
+					List<string> spineCandidates = sources.Where(s => isTargeted.Contains(s)).ToList();
+					if (spineCandidates.Count == 1)
+					{
+						reverseMap[target] = spineCandidates[0];
+					}
+					else
+					{
+						// Ambiguous: pick first and warn
+						reverseMap[target] = sources[0];
+						logger?.LogWarning("Bot {BotName}: multiple branches target '{Target}' via forceFlowTo ({Sources}), using '{Chosen}'",
+							graph.BotName, target, String.Join(", ", sources), sources[0]);
+					}
+				}
+			}
+
 			// Find root(s) — branches that appear as forceFlowTo targets but have no forceFlowTo of their own
 			List<string> roots = new List<string>();
 			foreach (string target in reverseMap.Keys)
