Horde Perforce Materalizer - Fix for swapping edge server not doing a full clean of the workspace

diff --git a/Engine/Source/Programs/Horde/Drivers/JobDriver/Execution/PerforceMaterializer.cs b/Engine/Source/Programs/Horde/Drivers/JobDriver/Execution/PerforceMaterializer.cs
--- a/Engine/Source/Programs/Horde/Drivers/JobDriver/Execution/PerforceMaterializer.cs
+++ b/Engine/Source/Programs/Horde/Drivers/JobDriver/Execution/PerforceMaterializer.cs
@@ -281,12 +281,22 @@ public sealed class PerforceMaterializer : IWorkspaceMaterializer
 		_lastSyncOptions = options;
 		State? state = await LoadStateAsync(cancellationToken);
 		
+		// Check if only the client name changed (edge server switch scenario)
+		// In this case, don't mark as dirty - the existing flush + sync will handle it correctly
+		bool isClientNameChange = state != null
+			&& state.Version == CurrentStateVersion
+			&& state.Status != TransactionStatus.Dirty
+			&& state.Identifier == Identifier
+			&& state.Stream == _options.AgentWorkspace.Stream
+			&& state.Client != _perforceSettings!.ClientName
+			&& state.Changelist != -1;
+ 
 		bool isDirty = state == null
 			|| state.Version != CurrentStateVersion
 			|| state.Status == TransactionStatus.Dirty
 			|| state.Identifier != Identifier
 			|| state.Stream != _options.AgentWorkspace.Stream
-			|| state.Client != _perforceSettings!.ClientName
+			|| (!isClientNameChange && state.Client != _perforceSettings!.ClientName)
 			|| state.Changelist == -1;
 		
 		_logger.LogInformation("Status: {Status}\nVersion: {PrevVersion} -> {Version}\nClient: {PrevClient} -> {Client}\nIdentifier: {PrevId} -> {Id}\nStream: {PrevStream} -> {Stream}\nCL: {PrevCl} -> {Cl}\nShelve: {PrevShelve} -> {Shelve}\nDirty: {Dirty}",
@@ -334,7 +344,7 @@ public sealed class PerforceMaterializer : IWorkspaceMaterializer
 		_perforceWithClient = await PerforceConnection.CreateAsync(_perforceSettings, _logger);
 		await RevertInternalAsync(_perforceWithClient, cancellationToken);
 
-		if (clientWasCreated && !isDirty && state != null)
+		if ((clientWasCreated || isClientNameChange) && !isDirty && state != null)
 		{
 			using TelemetrySpan deleteSpan = _tracer.StartActiveSpan($"{nameof(PerforceMaterializer)}.{nameof(SyncAsync)}.Flush");
 			using ILoggerProgress status = _logger.BeginProgressScope($"Flushing client to CL {state.Changelist} (p4 sync -k)");
