[Test Automation Hub] Limit the size of the test tags panel in the Test Health details

#rnx
#Synced-CL 52060281

diff --git a/Engine/Source/Programs/Horde/HordeDashboard/plugins/test-automation/testAutomationHealth.tsx b/Engine/Source/Programs/Horde/HordeDashboard/plugins/test-automation/testAutomationHealth.tsx
--- a/Engine/Source/Programs/Horde/HordeDashboard/plugins/test-automation/testAutomationHealth.tsx
+++ b/Engine/Source/Programs/Horde/HordeDashboard/plugins/test-automation/testAutomationHealth.tsx
@@ -771,9 +771,9 @@ export const TestHealthView: React.FC<{ test: TestHealthItem, handler: TestDataH
                                 </Stack>
 
                                 {tags &&
-                                    <Stack horizontal verticalAlign="center" tokens={{childrenGap: 4}} wrap grow style={{paddingLeft: 8}}>
+                                    <Stack horizontal verticalAlign="start" wrap grow style={{paddingLeft: 8, maxHeight: 160, overflowY: 'auto', overflowX: 'hidden'}}>
                                         {tags.map(tag =>
-                                                <Stack key={`tag-${tag}`} style={{backgroundColor: tagColors.getTagColor(tag), borderRadius: 5, padding: '1px 4px', opacity: 0.8}}>
+                                                <Stack key={`tag-${tag}`} style={{backgroundColor: tagColors.getTagColor(tag), borderRadius: 5, padding: '1px 4px', margin: '0px 2px 2px 0px', opacity: 0.8}}>
                                                     <Text>#{tag}</Text>
                                                 </Stack>
                                             )
@@ -784,7 +784,7 @@ export const TestHealthView: React.FC<{ test: TestHealthItem, handler: TestDataH
                             </Stack>
                             <Stack tokens={{childrenGap: 8}} style={{width: '100%'}}>
 
-                                <Stack horizontal verticalAlign="start" tokens={{childrenGap: 6}} style={{cursor: 'default'}} grow>
+                                <Stack horizontal verticalAlign="start" tokens={{childrenGap: 6}} style={{cursor: 'default'}}>
                                     <Label style={{whiteSpace: 'nowrap'}}>Degradation Notification: </Label>
                                     {!notificationOptions && <Text style={{color:textEmphasisColor(), cursor: 'default'}}>no workflow available in this stream</Text>}
                                     {!!notificationOptions &&
@@ -809,7 +809,7 @@ export const TestHealthView: React.FC<{ test: TestHealthItem, handler: TestDataH
                                     }
                                 </Stack>
 
-                                <Stack horizontal tokens={{childrenGap: 12}} grow>
+                                <Stack horizontal verticalAlign="start" tokens={{childrenGap: 12}} grow>
                                     <Stack tokens={{childrenGap: 6}} style={{borderRadius: 4}} verticalAlign="start">
                                         <Label>Owner: </Label>
                                         <Stack horizontal tokens={{childrenGap: 6}} wrap>
