Horde: Updating schema documentation

#rnx

diff --git a/Engine/Source/Programs/Horde/Docs/Config/Schema/AclActions.md b/Engine/Source/Programs/Horde/Docs/Config/Schema/AclActions.md
--- a/Engine/Source/Programs/Horde/Docs/Config/Schema/AclActions.md
+++ b/Engine/Source/Programs/Horde/Docs/Config/Schema/AclActions.md
@@ -63,48 +63,6 @@
 | `DeleteAccount` |  |
 | `ViewAccount` |  |
 
-## Artifacts
-
-| Name | Description |
-| ---- | ----------- |
-| `ReadArtifact` | Permission to read from an artifact |
-| `WriteArtifact` | Permission to write to an artifact |
-| `DeleteArtifact` | Permission to delete to an artifact |
-| `UploadArtifact` | Ability to create an artifact. Typically just for debugging; agents have this access for a particular session. |
-| `DownloadArtifact` | Ability to download an artifact |
-
-## Bisect
-
-| Name | Description |
-| ---- | ----------- |
-| `CreateBisectTask` | Ability to start new bisect tasks |
-| `UpdateBisectTask` | Ability to update a bisect task |
-| `ViewBisectTask` | Ability to view a bisect task |
-
-## BuildHealth
-
-| Name | Description |
-| ---- | ----------- |
-| `AddUpdateDeleteBuildHealthFilter` | Ability to add, update, and delete build health filters. |
-
-## Devices
-
-| Name | Description |
-| ---- | ----------- |
-| `DeviceRead` | Ability to read devices |
-| `DeviceWrite` | Ability to write devices |
-
-## Jobs
-
-| Name | Description |
-| ---- | ----------- |
-| `CreateJob` | Ability to start new jobs |
-| `UpdateJob` | Rename a job, modify its priority, etc... |
-| `DeleteJob` | Delete a job properties |
-| `ExecuteJob` | Allows updating a job metadata (name, changelist number, step properties, new groups, job states, etc...). Typically granted to agents. Not user facing. |
-| `RetryJobStep` | Ability to retry a failed job step |
-| `ViewJob` | Ability to view a job |
-
 ## Notices
 
 | Name | Description |
@@ -113,28 +71,6 @@
 | `UpdateNotice` |  |
 | `DeleteNotice` |  |
 
-## Notifications
-
-| Name | Description |
-| ---- | ----------- |
-| `CreateSubscription` | Ability to subscribe to notifications |
-
-## Projects
-
-| Name | Description |
-| ---- | ----------- |
-| `CreateProject` | Allows the creation of new projects |
-| `DeleteProject` | Allows deletion of projects. |
-| `UpdateProject` | Modify attributes of a project (name, categories, etc...) |
-| `ViewProject` | View information about a project |
-
-## Replicators
-
-| Name | Description |
-| ---- | ----------- |
-| `UpdateReplicator` | Allows deletion of projects. |
-| `ViewReplicator` | Allows the creation of new projects |
-
 ## ServiceAccounts
 
 | Name | Description |
@@ -144,24 +80,6 @@
 | `DeleteAccount` |  |
 | `ViewAccount` |  |
 
-## Streams
-
-| Name | Description |
-| ---- | ----------- |
-| `CreateStream` | Allows the creation of new streams within a project |
-| `UpdateStream` | Allows updating a stream (agent types, templates, schedules) |
-| `DeleteStream` | Allows deleting a stream |
-| `ViewStream` | Ability to view a stream |
-| `ViewChanges` | View changes submitted to a stream. NOTE: this returns responses from the server's Perforce account, which may be a priviledged user. |
-| `ViewTemplate` | View template associated with a stream |
-
-## TestData
-
-| Name | Description |
-| ---- | ----------- |
-| `TestHealthRead` | Ability to read test health |
-| `TestHealthWrite` | Ability to write test health |
-
 ## Accounts
 
 | Name | Description |
@@ -171,47 +89,37 @@
 | `DeleteAccount` |  |
 | `ViewAccount` |  |
 
-## Artifacts
-
-| Name | Description |
-| ---- | ----------- |
-| `ReadArtifact` |  |
-| `WriteArtifact` |  |
-| `DeleteArtifact` |  |
-| `UploadArtifact` |  |
-| `DownloadArtifact` |  |
-
-## Bisect
+## Notices
 
 | Name | Description |
 | ---- | ----------- |
-| `CreateBisectTask` |  |
-| `UpdateBisectTask` |  |
-| `ViewBisectTask` |  |
+| `CreateNotice` |  |
+| `UpdateNotice` |  |
+| `DeleteNotice` |  |
 
-## BuildHealth
+## Secrets
 
 | Name | Description |
 | ---- | ----------- |
-| `AddUpdateDeleteBuildHealthFilter` |  |
+| `ViewSecret` | View a credential |
 
-## Devices
+## ServiceAccounts
 
 | Name | Description |
 | ---- | ----------- |
-| `DeviceRead` |  |
-| `DeviceWrite` |  |
+| `CreateAccount` |  |
+| `UpdateAccount` |  |
+| `DeleteAccount` |  |
+| `ViewAccount` |  |
 
-## Jobs
+## Accounts
 
 | Name | Description |
 | ---- | ----------- |
-| `CreateJob` |  |
-| `UpdateJob` |  |
-| `DeleteJob` |  |
-| `ExecuteJob` |  |
-| `RetryJobStep` |  |
-| `ViewJob` |  |
+| `CreateAccount` |  |
+| `UpdateAccount` |  |
+| `DeleteAccount` |  |
+| `ViewAccount` |  |
 
 ## Notices
 
@@ -221,27 +129,11 @@
 | `UpdateNotice` |  |
 | `DeleteNotice` |  |
 
-## Notifications
-
-| Name | Description |
-| ---- | ----------- |
-| `CreateSubscription` |  |
-
-## Projects
-
-| Name | Description |
-| ---- | ----------- |
-| `CreateProject` |  |
-| `DeleteProject` |  |
-| `UpdateProject` |  |
-| `ViewProject` |  |
-
-## Replicators
+## Secrets
 
 | Name | Description |
 | ---- | ----------- |
-| `UpdateReplicator` |  |
-| `ViewReplicator` |  |
+| `ViewSecret` |  |
 
 ## ServiceAccounts
 
@@ -252,23 +144,18 @@
 | `DeleteAccount` |  |
 | `ViewAccount` |  |
 
-## Streams
-
-| Name | Description |
-| ---- | ----------- |
-| `CreateStream` |  |
-| `UpdateStream` |  |
-| `DeleteStream` |  |
-| `ViewStream` |  |
-| `ViewChanges` |  |
-| `ViewTemplate` |  |
-
-## TestData
+## Storage
 
 | Name | Description |
 | ---- | ----------- |
-| `TestHealthRead` |  |
-| `TestHealthWrite` |  |
+| `ReadBlobs` | Ability to read blobs from the storage service |
+| `WriteBlobs` | Ability to write blobs to the storage service |
+| `ReadAliases` | Ability to read aliases from the storage service |
+| `WriteAliases` | Ability to write aliases to the storage service |
+| `DeleteAliases` | Ability to write aliases to the storage service |
+| `ReadRefs` | Ability to read refs from the storage service |
+| `WriteRefs` | Ability to write refs to the storage service |
+| `DeleteRefs` | Ability to delete refs |
 
 ## Accounts
 
@@ -290,30 +177,6 @@
 | `ViewAgent` | View an agent |
 | `ListAgents` | List the available agents |
 
-## Artifacts
-
-| Name | Description |
-| ---- | ----------- |
-| `ReadArtifact` |  |
-| `WriteArtifact` |  |
-| `DeleteArtifact` |  |
-| `UploadArtifact` |  |
-| `DownloadArtifact` |  |
-
-## Bisect
-
-| Name | Description |
-| ---- | ----------- |
-| `CreateBisectTask` |  |
-| `UpdateBisectTask` |  |
-| `ViewBisectTask` |  |
-
-## BuildHealth
-
-| Name | Description |
-| ---- | ----------- |
-| `AddUpdateDeleteBuildHealthFilter` |  |
-
 ## Compute
 
 | Name | Description |
@@ -323,24 +186,6 @@
 | `UbaCacheRead` | User can read from Unreal Build Accelerator cache |
 | `UbaCacheWrite` | User can write to Unreal Build Accelerator cache |
 
-## Devices
-
-| Name | Description |
-| ---- | ----------- |
-| `DeviceRead` |  |
-| `DeviceWrite` |  |
-
-## Jobs
-
-| Name | Description |
-| ---- | ----------- |
-| `CreateJob` |  |
-| `UpdateJob` |  |
-| `DeleteJob` |  |
-| `ExecuteJob` |  |
-| `RetryJobStep` |  |
-| `ViewJob` |  |
-
 ## Leases
 
 | Name | Description |
@@ -367,208 +212,21 @@
 | `UpdateNotice` |  |
 | `DeleteNotice` |  |
 
-## Notifications
-
-| Name | Description |
-| ---- | ----------- |
-| `CreateSubscription` |  |
-
-## Pools
-
-| Name | Description |
-| ---- | ----------- |
-| `CreatePool` | Create a global pool of agents |
-| `UpdatePool` | Modify an agent pool |
-| `DeletePool` | Delete an agent pool |
-| `ViewPool` | Ability to view a pool |
-| `ListPools` | View all the available agent pools |
-
-## Projects
-
-| Name | Description |
-| ---- | ----------- |
-| `CreateProject` |  |
-| `DeleteProject` |  |
-| `UpdateProject` |  |
-| `ViewProject` |  |
-
-## Replicators
-
-| Name | Description |
-| ---- | ----------- |
-| `UpdateReplicator` |  |
-| `ViewReplicator` |  |
-
-## ServiceAccounts
-
-| Name | Description |
-| ---- | ----------- |
-| `CreateAccount` |  |
-| `UpdateAccount` |  |
-| `DeleteAccount` |  |
-| `ViewAccount` |  |
-
-## Sessions
-
-| Name | Description |
-| ---- | ----------- |
-| `CreateSession` | Granted to agents to call CreateSession, which returns a bearer token identifying themselves valid to call UpdateSesssion via gRPC. |
-| `ViewSession` | Allows viewing information about an agent session |
-
-## Software
-
-| Name | Description |
-| ---- | ----------- |
-| `UploadSoftware` | Ability to upload new versions of the agent software |
-| `DownloadSoftware` | Ability to download the agent software |
-| `DeleteSoftware` | Ability to delete agent software |
-
-## Streams
-
-| Name | Description |
-| ---- | ----------- |
-| `CreateStream` |  |
-| `UpdateStream` |  |
-| `DeleteStream` |  |
-| `ViewStream` |  |
-| `ViewChanges` |  |
-| `ViewTemplate` |  |
-
-## TestData
-
-| Name | Description |
-| ---- | ----------- |
-| `TestHealthRead` |  |
-| `TestHealthWrite` |  |
-
-## Accounts
-
-| Name | Description |
-| ---- | ----------- |
-| `CreateAccount` |  |
-| `UpdateAccount` |  |
-| `DeleteAccount` |  |
-| `ViewAccount` |  |
-
-## Agents
-
-| Name | Description |
-| ---- | ----------- |
-| `CreateAgent` |  |
-| `CreateWorkstationAgent` |  |
-| `UpdateAgent` |  |
-| `DeleteAgent` |  |
-| `ViewAgent` |  |
-| `ListAgents` |  |
-
-## Artifacts
-
-| Name | Description |
-| ---- | ----------- |
-| `ReadArtifact` |  |
-| `WriteArtifact` |  |
-| `DeleteArtifact` |  |
-| `UploadArtifact` |  |
-| `DownloadArtifact` |  |
-
-## Bisect
-
-| Name | Description |
-| ---- | ----------- |
-| `CreateBisectTask` |  |
-| `UpdateBisectTask` |  |
-| `ViewBisectTask` |  |
-
-## BuildHealth
-
-| Name | Description |
-| ---- | ----------- |
-| `AddUpdateDeleteBuildHealthFilter` |  |
-
-## Compute
-
-| Name | Description |
-| ---- | ----------- |
-| `AddComputeTasks` |  |
-| `GetComputeTasks` |  |
-| `UbaCacheRead` |  |
-| `UbaCacheWrite` |  |
-
-## Devices
-
-| Name | Description |
-| ---- | ----------- |
-| `DeviceRead` |  |
-| `DeviceWrite` |  |
-
-## Jobs
-
-| Name | Description |
-| ---- | ----------- |
-| `CreateJob` |  |
-| `UpdateJob` |  |
-| `DeleteJob` |  |
-| `ExecuteJob` |  |
-| `RetryJobStep` |  |
-| `ViewJob` |  |
-
-## Leases
-
-| Name | Description |
-| ---- | ----------- |
-| `ViewLeases` |  |
-| `ViewLeaseTasks` |  |
-
-## Logs
-
-| Name | Description |
-| ---- | ----------- |
-| `CreateLog` |  |
-| `UpdateLog` |  |
-| `ViewLog` |  |
-| `WriteLogData` |  |
-| `CreateEvent` |  |
-| `ViewEvent` |  |
-
-## Notices
-
-| Name | Description |
-| ---- | ----------- |
-| `CreateNotice` |  |
-| `UpdateNotice` |  |
-| `DeleteNotice` |  |
-
-## Notifications
-
-| Name | Description |
-| ---- | ----------- |
-| `CreateSubscription` |  |
-
 ## Pools
 
 | Name | Description |
 | ---- | ----------- |
-| `CreatePool` |  |
-| `UpdatePool` |  |
-| `DeletePool` |  |
-| `ViewPool` |  |
-| `ListPools` |  |
-
-## Projects
-
-| Name | Description |
-| ---- | ----------- |
-| `CreateProject` |  |
-| `DeleteProject` |  |
-| `UpdateProject` |  |
-| `ViewProject` |  |
+| `CreatePool` | Create a global pool of agents |
+| `UpdatePool` | Modify an agent pool |
+| `DeletePool` | Delete an agent pool |
+| `ViewPool` | Ability to view a pool |
+| `ListPools` | View all the available agent pools |
 
-## Replicators
+## Secrets
 
 | Name | Description |
 | ---- | ----------- |
-| `UpdateReplicator` |  |
-| `ViewReplicator` |  |
+| `ViewSecret` |  |
 
 ## ServiceAccounts
 
@@ -583,34 +241,29 @@
 
 | Name | Description |
 | ---- | ----------- |
-| `CreateSession` |  |
-| `ViewSession` |  |
+| `CreateSession` | Granted to agents to call CreateSession, which returns a bearer token identifying themselves valid to call UpdateSesssion via gRPC. |
+| `ViewSession` | Allows viewing information about an agent session |
 
 ## Software
 
 | Name | Description |
 | ---- | ----------- |
-| `UploadSoftware` |  |
-| `DownloadSoftware` |  |
-| `DeleteSoftware` |  |
-
-## Streams
-
-| Name | Description |
-| ---- | ----------- |
-| `CreateStream` |  |
-| `UpdateStream` |  |
-| `DeleteStream` |  |
-| `ViewStream` |  |
-| `ViewChanges` |  |
-| `ViewTemplate` |  |
+| `UploadSoftware` | Ability to upload new versions of the agent software |
+| `DownloadSoftware` | Ability to download the agent software |
+| `DeleteSoftware` | Ability to delete agent software |
 
-## TestData
+## Storage
 
 | Name | Description |
 | ---- | ----------- |
-| `TestHealthRead` |  |
-| `TestHealthWrite` |  |
+| `ReadBlobs` |  |
+| `WriteBlobs` |  |
+| `ReadAliases` |  |
+| `WriteAliases` |  |
+| `DeleteAliases` |  |
+| `ReadRefs` |  |
+| `WriteRefs` |  |
+| `DeleteRefs` |  |
 
 ## Accounts
 
@@ -632,30 +285,6 @@
 | `ViewAgent` |  |
 | `ListAgents` |  |
 
-## Artifacts
-
-| Name | Description |
-| ---- | ----------- |
-| `ReadArtifact` |  |
-| `WriteArtifact` |  |
-| `DeleteArtifact` |  |
-| `UploadArtifact` |  |
-| `DownloadArtifact` |  |
-
-## Bisect
-
-| Name | Description |
-| ---- | ----------- |
-| `CreateBisectTask` |  |
-| `UpdateBisectTask` |  |
-| `ViewBisectTask` |  |
-
-## BuildHealth
-
-| Name | Description |
-| ---- | ----------- |
-| `AddUpdateDeleteBuildHealthFilter` |  |
-
 ## Compute
 
 | Name | Description |
@@ -665,24 +294,6 @@
 | `UbaCacheRead` |  |
 | `UbaCacheWrite` |  |
 
-## Devices
-
-| Name | Description |
-| ---- | ----------- |
-| `DeviceRead` |  |
-| `DeviceWrite` |  |
-
-## Jobs
-
-| Name | Description |
-| ---- | ----------- |
-| `CreateJob` |  |
-| `UpdateJob` |  |
-| `DeleteJob` |  |
-| `ExecuteJob` |  |
-| `RetryJobStep` |  |
-| `ViewJob` |  |
-
 ## Leases
 
 | Name | Description |
@@ -709,12 +320,6 @@
 | `UpdateNotice` |  |
 | `DeleteNotice` |  |
 
-## Notifications
-
-| Name | Description |
-| ---- | ----------- |
-| `CreateSubscription` |  |
-
 ## Pools
 
 | Name | Description |
@@ -725,27 +330,11 @@
 | `ViewPool` |  |
 | `ListPools` |  |
 
-## Projects
-
-| Name | Description |
-| ---- | ----------- |
-| `CreateProject` |  |
-| `DeleteProject` |  |
-| `UpdateProject` |  |
-| `ViewProject` |  |
-
-## Replicators
-
-| Name | Description |
-| ---- | ----------- |
-| `UpdateReplicator` |  |
-| `ViewReplicator` |  |
-
 ## Secrets
 
 | Name | Description |
 | ---- | ----------- |
-| `ViewSecret` | View a credential |
+| `ViewSecret` |  |
 
 ## ServiceAccounts
 
@@ -771,23 +360,24 @@
 | `DownloadSoftware` |  |
 | `DeleteSoftware` |  |
 
-## Streams
+## Storage
 
 | Name | Description |
 | ---- | ----------- |
-| `CreateStream` |  |
-| `UpdateStream` |  |
-| `DeleteStream` |  |
-| `ViewStream` |  |
-| `ViewChanges` |  |
-| `ViewTemplate` |  |
+| `ReadBlobs` |  |
+| `WriteBlobs` |  |
+| `ReadAliases` |  |
+| `WriteAliases` |  |
+| `DeleteAliases` |  |
+| `ReadRefs` |  |
+| `WriteRefs` |  |
+| `DeleteRefs` |  |
 
-## TestData
+## Symbols
 
 | Name | Description |
 | ---- | ----------- |
-| `TestHealthRead` |  |
-| `TestHealthWrite` |  |
+| `ReadSymbols` | Ability to download symbols |
 
 ## Accounts
 
@@ -813,25 +403,25 @@
 
 | Name | Description |
 | ---- | ----------- |
-| `ReadArtifact` |  |
-| `WriteArtifact` |  |
-| `DeleteArtifact` |  |
-| `UploadArtifact` |  |
-| `DownloadArtifact` |  |
+| `ReadArtifact` | Permission to read from an artifact |
+| `WriteArtifact` | Permission to write to an artifact |
+| `DeleteArtifact` | Permission to delete to an artifact |
+| `UploadArtifact` | Ability to create an artifact. Typically just for debugging; agents have this access for a particular session. |
+| `DownloadArtifact` | Ability to download an artifact |
 
 ## Bisect
 
 | Name | Description |
 | ---- | ----------- |
-| `CreateBisectTask` |  |
-| `UpdateBisectTask` |  |
-| `ViewBisectTask` |  |
+| `CreateBisectTask` | Ability to start new bisect tasks |
+| `UpdateBisectTask` | Ability to update a bisect task |
+| `ViewBisectTask` | Ability to view a bisect task |
 
 ## BuildHealth
 
 | Name | Description |
 | ---- | ----------- |
-| `AddUpdateDeleteBuildHealthFilter` |  |
+| `AddUpdateDeleteBuildHealthFilter` | Ability to add, update, and delete build health filters. |
 
 ## Compute
 
@@ -846,19 +436,19 @@
 
 | Name | Description |
 | ---- | ----------- |
-| `DeviceRead` |  |
-| `DeviceWrite` |  |
+| `DeviceRead` | Ability to read devices |
+| `DeviceWrite` | Ability to write devices |
 
 ## Jobs
 
 | Name | Description |
 | ---- | ----------- |
-| `CreateJob` |  |
-| `UpdateJob` |  |
-| `DeleteJob` |  |
-| `ExecuteJob` |  |
-| `RetryJobStep` |  |
-| `ViewJob` |  |
+| `CreateJob` | Ability to start new jobs |
+| `UpdateJob` | Rename a job, modify its priority, etc... |
+| `DeleteJob` | Delete a job properties |
+| `ExecuteJob` | Allows updating a job metadata (name, changelist number, step properties, new groups, job states, etc...). Typically granted to agents. Not user facing. |
+| `RetryJobStep` | Ability to retry a failed job step |
+| `ViewJob` | Ability to view a job |
 
 ## Leases
 
@@ -890,7 +480,7 @@
 
 | Name | Description |
 | ---- | ----------- |
-| `CreateSubscription` |  |
+| `CreateSubscription` | Ability to subscribe to notifications |
 
 ## Pools
 
@@ -906,17 +496,17 @@
 
 | Name | Description |
 | ---- | ----------- |
-| `CreateProject` |  |
-| `DeleteProject` |  |
-| `UpdateProject` |  |
-| `ViewProject` |  |
+| `CreateProject` | Allows the creation of new projects |
+| `DeleteProject` | Allows deletion of projects. |
+| `UpdateProject` | Modify attributes of a project (name, categories, etc...) |
+| `ViewProject` | View information about a project |
 
 ## Replicators
 
 | Name | Description |
 | ---- | ----------- |
-| `UpdateReplicator` |  |
-| `ViewReplicator` |  |
+| `UpdateReplicator` | Allows deletion of projects. |
+| `ViewReplicator` | Allows the creation of new projects |
 
 ## Secrets
 
@@ -948,29 +538,42 @@
 | `DownloadSoftware` |  |
 | `DeleteSoftware` |  |
 
+## Storage
+
+| Name | Description |
+| ---- | ----------- |
+| `ReadBlobs` |  |
+| `WriteBlobs` |  |
+| `ReadAliases` |  |
+| `WriteAliases` |  |
+| `DeleteAliases` |  |
+| `ReadRefs` |  |
+| `WriteRefs` |  |
+| `DeleteRefs` |  |
+
 ## Streams
 
 | Name | Description |
 | ---- | ----------- |
-| `CreateStream` |  |
-| `UpdateStream` |  |
-| `DeleteStream` |  |
-| `ViewStream` |  |
-| `ViewChanges` |  |
-| `ViewTemplate` |  |
+| `CreateStream` | Allows the creation of new streams within a project |
+| `UpdateStream` | Allows updating a stream (agent types, templates, schedules) |
+| `DeleteStream` | Allows deleting a stream |
+| `ViewStream` | Ability to view a stream |
+| `ViewChanges` | View changes submitted to a stream. NOTE: this returns responses from the server's Perforce account, which may be a priviledged user. |
+| `ViewTemplate` | View template associated with a stream |
 
 ## Symbols
 
 | Name | Description |
 | ---- | ----------- |
-| `ReadSymbols` | Ability to download symbols |
+| `ReadSymbols` |  |
 
 ## TestData
 
 | Name | Description |
 | ---- | ----------- |
-| `TestHealthRead` |  |
-| `TestHealthWrite` |  |
+| `TestHealthRead` | Ability to read test health |
+| `TestHealthWrite` | Ability to write test health |
 
 ## Accounts
 
@@ -1135,14 +738,14 @@
 
 | Name | Description |
 | ---- | ----------- |
-| `ReadBlobs` | Ability to read blobs from the storage service |
-| `WriteBlobs` | Ability to write blobs to the storage service |
-| `ReadAliases` | Ability to read aliases from the storage service |
-| `WriteAliases` | Ability to write aliases to the storage service |
-| `DeleteAliases` | Ability to write aliases to the storage service |
-| `ReadRefs` | Ability to read refs from the storage service |
-| `WriteRefs` | Ability to write refs to the storage service |
-| `DeleteRefs` | Ability to delete refs |
+| `ReadBlobs` |  |
+| `WriteBlobs` |  |
+| `ReadAliases` |  |
+| `WriteAliases` |  |
+| `DeleteAliases` |  |
+| `ReadRefs` |  |
+| `WriteRefs` |  |
+| `DeleteRefs` |  |
 
 ## Streams
 
@@ -1168,6 +771,13 @@
 | `TestHealthRead` |  |
 | `TestHealthWrite` |  |
 
+## Tools
+
+| Name | Description |
+| ---- | ----------- |
+| `DownloadTool` | Ability to download a tool |
+| `UploadTool` | Ability to upload new tool versions |
+
 ## Accounts
 
 | Name | Description |
@@ -1368,8 +978,8 @@
 
 | Name | Description |
 | ---- | ----------- |
-| `DownloadTool` | Ability to download a tool |
-| `UploadTool` | Ability to upload new tool versions |
+| `DownloadTool` |  |
+| `UploadTool` |  |
 
 ## Accounts
 
diff --git a/Engine/Source/Programs/Horde/Docs/Config/Schema/Globals.md b/Engine/Source/Programs/Horde/Docs/Config/Schema/Globals.md
--- a/Engine/Source/Programs/Horde/Docs/Config/Schema/Globals.md
+++ b/Engine/Source/Programs/Horde/Docs/Config/Schema/Globals.md
@@ -60,14 +60,14 @@
 Name | Description
 ---- | -----------
 `analytics` | [AnalyticsConfig](#analyticsconfig)<br>Configuration for the analytics plugin
-`build` | [BuildConfig](#buildconfig)<br>Configuration for the build plugin
 `exampleplugin` | [ExamplePluginConfig](#examplepluginconfig)<br>Configuration for the exampleplugin plugin
-`compute` | [ComputeConfig](#computeconfig)<br>Configuration for the compute plugin
-`experimental` | [ExperimentalConfig](#experimentalconfig)<br>Configuration for the experimental plugin
 `secrets` | [SecretsConfig](#secretsconfig)<br>Configuration for the secrets plugin
-`symbols` | [SymbolsConfig](#symbolsconfig)<br>Configuration for the symbols plugin
 `storage` | [StorageConfig](#storageconfig)<br>Configuration for the storage plugin
+`compute` | [ComputeConfig](#computeconfig)<br>Configuration for the compute plugin
+`symbols` | [SymbolsConfig](#symbolsconfig)<br>Configuration for the symbols plugin
+`build` | [BuildConfig](#buildconfig)<br>Configuration for the build plugin
 `tools` | [ToolsConfig](#toolsconfig)<br>Configuration for the tools plugin
+`experimental` | [ExperimentalConfig](#experimentalconfig)<br>Configuration for the experimental plugin
 
 ## AnalyticsConfig
 
@@ -77,153 +77,55 @@
 ---- | -----------
 `stores` | [TelemetryStoreConfig](Telemetry.md)`[]`<br>Metrics to aggregate on the Horde server
 
-## BuildConfig
+## ExamplePluginConfig
 
-Configuration for the build plugin
+Configuration for the ExamplePlugin plugin
 
 Name | Description
 ---- | -----------
-`perforceClusters` | [PerforceCluster](#perforcecluster)`[]`<br>List of Perforce clusters
-`devices` | [DeviceConfig](#deviceconfig)<br>Device configuration
-`maxConformCount` | `integer`<br>Maximum number of conforms to run at once
-`agentShutdownIfDisabledGracePeriod` | `string`<br>Time to wait before shutting down an agent that has been disabled Used if no value is set on the actual pool.
-`robomerge` | [RoboMergeConfig](#robomergeconfig)<br>RoboMerge branchmap file configuration for the merge graph service. When null or empty, the service is inert (no Perforce connections, no API data).
-`artifactTypes` | [ArtifactTypeConfig](#artifacttypeconfig)`[]`<br>Configuration for different artifact types
-`projects` | [ProjectConfig](Projects.md)`[]`<br>List of projects
-`enableConformTasks` | `boolean`<br>Whether to allow conform tasks to run
-`enableJobTaskConformCheck` | `boolean`<br>Whether to check if a conform is required on an agent before assigning a lease for a job task source.
-`issueFixedTag` | `string`<br>Commit tag to use for marking issues as fixed
-`issues` | [IssueConfig](Issues.md)<br>Build health issue configurations
-`enableIssueRootCauseAnalysisNotifications` | `boolean`<br>Whether to allow notifications for Root Cause Analysis to be sent
-`enableAutoRcaSlackNotifications` | `boolean`<br>Global kill switch for posting auto-RCA results to Slack triage threads. When false (default), no auto-RCA Slack notifications are sent regardless of per-workflow AutoRcaSlackEnabled settings.
-`autoRcaUseMultiStage` | `boolean`<br>Whether to use multi-stage LLM analysis for auto-RCA. Multi-stage is slower but potentially more accurate. Set to false for faster processing with higher throughput.
-`maxConcurrentRcaRequests` | `integer`<br>Maximum number of concurrent auto-RCA requests. Limits load on the external LLM API. Set to 0 to disable auto-RCA globally (acts as a kill switch).
-`issueFindSuspectsMaxSuccessDays` | `integer`<br>The maximum days since node success to scan for suspects
-`issueFindSuspectsMaxScoringSeconds` | `integer`<br>The maximum number of seconds to spend scording suspects
-`issueFindSuspectsMaxChanges` | `integer`<br>Maximum number of commits to examine for suspects. Previously hardcoded as const MaxChanges = 500.
-`issueFindSuspectsEnableParallelScoring` | `boolean`<br>Enable parallel scoring of commits for suspect finding. When enabled, commits are scored concurrently up to MaxParallelism.
-`issueFindSuspectsMaxParallelism` | `integer`<br>Maximum concurrent commit file lookups when parallel scoring is enabled. Higher values improve speed but increase Perforce server load.
-`issueFindSuspectsEnableEarlyExit` | `boolean`<br>Enable early exit when theoretical maximum score is achieved. Only exits early if the best possible score has been found.
-`issueFindSuspectsEnableBidirectionalQuery` | `boolean`<br>Enable bidirectional commit query to capture both recent and old commits. Queries half the limit from each end of the commit range.
-`issueFindSuspectsEnableFingerprintFallback` | `boolean`<br>Enable fingerprint-based fallback when no suspects are found. Helps prevent duplicate issues when commit range exceeds MaxChanges.
-`issueFingerprintFallbackMinScore` | `integer`<br>Minimum score required for fingerprint fallback matching. Higher values are more conservative (fewer false matches).
-`issueFingerprintFallbackMaxAgeDays` | `integer`<br>Maximum age in days for issues to be considered in fingerprint fallback.
-`issueFingerprintFallbackExcludedTypes` | `string[]`<br>Fingerprint types excluded from fallback matching (too generic).
-`enableExpandedIgnorePatterns` | `boolean`<br>Whether to enable expanded IgnorePatterns discovery (workspace-wide scanning + config-specified files). When false (default), only Engine directory is scanned for ignore patterns.
-`lowContextFingerprintTypes` | `string[]`<br>Fingerprint type prefixes considered "low context" for suspect mention suppression. When an issue's fingerprint type starts with any of these prefixes AND all suspects are tied at rank 0, @ mentions are suppressed in Slack. Suspects are still listed without mentions. Set to an empty array to disable suppression entirely.
+`greetingMessage` | `string`<br>The Greeting provided by the example plugin
 
-## PerforceCluster
+## SecretsConfig
 
-Information about a cluster of Perforce servers.
+Configuration for the secrets system
 
 Name | Description
 ---- | -----------
-`name` | `string`<br>Name of the cluster
-`serviceAccount` | `string`<br>Username for Horde to log in to this server. Will use the first account specified below if not overridden.
-`canImpersonate` | `boolean`<br>Whether the service account can impersonate other users
-`supportsPartitionedWorkspaces` | `boolean`<br>Whether to use partitioned workspaces on this server
-`servers` | [PerforceServer](#perforceserver)`[]`<br>List of servers
-`credentials` | [PerforceCredentials](#perforcecredentials)`[]`<br>List of server credentials
-`autoSdk` | [AutoSdkWorkspace](#autosdkworkspace)`[]`<br>List of autosdk streams
+`secrets` | [SecretConfig](#secretconfig)`[]`<br>List of secrets
+`providerConfigs` | [SecretProviderConfig](#secretproviderconfig)`[]`<br>List of configurations for secret providers
 
-## PerforceServer
+## SecretConfig
 
-Information about an individual Perforce server
+Configuration for a secret value
 
 Name | Description
 ---- | -----------
-`serverAndPort` | `string`<br>The server and port. The server may be a DNS entry with multiple records, in which case it will be actively load balanced. If "ssl:" prefix is used, ensure P4 server's fingerprint/certificate is trusted. See Horde's documentation on connecting to SSL-enabled Perforce servers.
-`healthCheck` | `boolean`<br>Whether to query the healthcheck address under each server
-`resolveDns` | `boolean`<br>Whether to resolve the DNS entries and load balance between different hosts
-`maxConformCount` | `integer`<br>Maximum number of simultaneous conforms on this server
-`condition` | `string`<br>Optional condition for a machine to be eligible to use this server
-`properties` | `string[]`<br>List of properties for an agent to be eligible to use this server
+`id` | `string`<br>Identifier for this secret
+`data` | `string` `->` `string`<br>Key/value pairs associated with this secret
+`sources` | [ExternalSecretConfig](#externalsecretconfig)`[]`<br>Providers to source key/value pairs from
+`acl` | [AclConfig](#aclconfig)<br>Defines access to this particular secret
 
-## PerforceCredentials
+## ExternalSecretConfig
 
-Credentials for a Perforce user
+Configuration for an external secret provider
 
 Name | Description
 ---- | -----------
-`userName` | `string`<br>The username
-`password` | `string`<br>Password for the user
-`ticket` | `string`<br>Login ticket for the user (will be used instead of password if set)
+`providerConfig` | `string`<br>The name to look up a secret provider config. A config is mutually exclusive to and to be used when a secret requires additional options than just the name of the provider
+`provider` | `string`<br>Name of the provider to use
+`format` | [ExternalSecretFormat](#externalsecretformat-enum)<br>Format of the secret
+`key` | `string`<br>Optional key indicating the parameter to set in the resulting data array. Required if  is .
+`path` | `string`<br>Optional value indicating what to fetch from the provider
 
-## AutoSdkWorkspace
+## ExternalSecretFormat (Enum)
 
-Path to a platform and stream to use for syncing AutoSDK
+Format describing how to parse external secret values
 
 Name | Description
 ---- | -----------
-`name` | `string`<br>Name of this workspace
-`properties` | `string[]`<br>The agent properties to check (eg. "OSFamily=Windows")
-`userName` | `string`<br>Username for logging in to the server
-`stream` | `string`<br>Stream to use
+`Text` | Secret is a plain text value which will be stored using the external secret key
+`Json` | Secret is a JSON formatted string containing key/value pairs
 
-## DeviceConfig
-
-Configuration for devices
-
-Name | Description
----- | -----------
-`platforms` | [DevicePlatformConfig](#deviceplatformconfig)`[]`<br>List of device platforms
-`pools` | [DevicePoolConfig](#devicepoolconfig)`[]`<br>List of device pools
-
-## DevicePlatformConfig
-
-Configuration for a device platform
-
-Name | Description
----- | -----------
-`id` | `string`<br>The id for this platform
-`name` | `string`<br>Name of the platform
-`models` | `string[]`<br>A list of platform models
-`legacyNames` | `string[]`<br>Legacy names which older versions of Gauntlet may be using
-`legacyPerfSpecHighModel` | `string`<br>Model name for the high perf spec, which may be requested by Gauntlet
-
-## DevicePoolConfig
-
-Configuration for a device pool
-
-Name | Description
----- | -----------
-`id` | `string`<br>The id for this platform
-`name` | `string`<br>The name of the pool
-`poolType` | [DevicePoolType](#devicepooltype-enum)<br>The type of the pool
-`projectIds` | `string[]`<br>List of project ids associated with pool
-
-## DevicePoolType (Enum)
-
-The type of device pool
-
-Name | Description
----- | -----------
-`Automation` | Available to CIS jobs
-`Shared` | Shared by users with remote checking and checkouts
-
-## RoboMergeConfig
-
-Configuration for the RoboMerge branchmap service
-
-Name | Description
----- | -----------
-`branchmapFiles` | `string[]`<br>Perforce depot paths to branchmap JSON files. Plain depot paths (e.g. "//MyProject/Main/Build/Streams/myproject.branchmap.json"). All files are read from the Perforce cluster specified by .
-`perforceCluster` | `string`<br>Perforce cluster name to connect to for reading branchmap files. Maps to a connection configured in BuildServerConfig.Perforce[]. Default: "default" (uses the default Perforce connection).
-`tickInterval` | `string`<br>How often to poll Perforce for branchmap file changes. Default: 5 minutes. Changing this requires a server restart (ticker interval is set at construction time).
-
-## ArtifactTypeConfig
-
-Configuration for an artifact
-
-Name | Description
----- | -----------
-`name` | `string`<br>Legacy 'Name' property
-`type` | `string`<br>Name of the artifact type
-`acl` | [AclConfig](#aclconfig)<br>Acl for the artifact type
-`keepCount` | `integer`<br>Number of artifacts to retain
-`keepDays` | `integer`<br>Number of days to retain artifacts of this type
-`namespaceId` | `string`<br>Storage namespace to use for this artifact types
-
 ## AclConfig
 
 Parameters to update an ACL
@@ -265,13 +167,124 @@
 `excludeActions` | `string[]`<br>Actions to exclude from the inherited actions
 `extends` | `string[]`<br>Other profiles to extend from
 
-## ExamplePluginConfig
+## SecretProviderConfig
 
-Configuration for the ExamplePlugin plugin
+Common settings object for different secret providers
 
 Name | Description
 ---- | -----------
-`greetingMessage` | `string`<br>The Greeting provided by the example plugin
+`name` | `string`<br>Name of the secret provider config
+`provider` | `string`<br>Name of the provider
+`hcpVault` | [HcpVaultConfig](#hcpvaultconfig)<br>Configuration for HCP Vault
+`azureKeyVault` | [AzureKeyVaultConfig](#azurekeyvaultconfig)<br>Configuration for Azure Key Vault
 
+## HcpVaultConfig
+
+Configuration for HCP Vault secrets.
+
+Name | Description
+---- | -----------
+`preSharedKey` | `string`<br>The Vault authentication token to use when the type of credentials is .<br>More information on this token can be found in the API documentation for X-Vault-Token
+`endPoint` | `string`<br>The address of the Vault server.
+`awsArnRole` | `string`<br>The optional AWS ARN to assume a role when logging into Vault. To be used when the type of credentials is . When not set the AWS SDK default credential search will be used.
+`awsIamServerId` | `string`<br>The value for the X-Vault-AWS-IAM-Server-ID header when the Vault server is configured with iam_server_id_header_value. Most likely to be the Vault server's DNS name. To be used when the type of credentials is .
+`role` | `string`<br>The name of the role in the Vault configuration for the login. This role name can be different to the name of the AWS IAM principal. To be used when the type of credentials is . When not set the name of IAM principal will be used.
+`credentials` | [HcpVaultCredentialsType](#hcpvaultcredentialstype-enum)<br>Type of credentials to use.
+
+## HcpVaultCredentialsType (Enum)
+
+Credentials to use to log into HCP Vault.
+
+Name | Description
+---- | -----------
+`PreSharedKey` | Use a pre-shared token provided by Vault for authentication.
+`AwsAuth` | Use AWS IAM credentials to authenticate against the Vault server.
+
+## AzureKeyVaultConfig
+
+Configuration for Azure Key Vault secrets.
+
+Name | Description
+---- | -----------
+`vaultUri` | `string`<br>The Uri to the Azure Key Vault.
+
+## StorageConfig
+
+Configuration for storage
+
+Name | Description
+---- | -----------
+`enableGc` | `boolean`<br>Whether to enable garbage collection
+`enableGcVerification` | `boolean`<br>Whether to enable garbage collection in verification mode (nothing deleted, just logging on access to deleted blobs)
+`gcWorkerCount` | `integer`<br>Number of concurrent GC reachability-check workers. Default 8. Reachability checks are I/O-bound (MongoDB queries), so higher values scale throughput roughly linearly up to the database's capacity.
+`gcQueueLimit` | `integer`<br>Maximum number of entries in a namespace's Redis GC queue before the blob scanner pauses. Default 50,000.
+`enableReliableGcEnqueue` | `boolean`<br>Whether to await Redis writes when enqueuing blobs for GC checking. When false (default), uses fire-and-forget which can silently drop entries on Redis failures. Set true for reliable enqueue.
+`enablePerNamespaceQueueIsolation` | `boolean`<br>Whether to isolate GC queue limits per-namespace. When false (default), one backed-up namespace pauses the scanner for ALL namespaces. When true, only the backed-up namespace is paused and its blobs are deferred.
+`gcBlobIngestionBatchSize` | `integer`<br>Number of blobs to process per blob-ingestion tick (TickBlobsAsync). Default 500. Raise to increase scan throughput when catching up on a backlog of unscanned blobs.
+`gcBatchSize` | `integer`<br>Number of blob IDs to batch together for the import-reachability check. Default 1 (current per-blob behavior). When greater than 1, uses $in query to check multiple blobs in a single MongoDB round-trip. Values of 50-200 recommended.
+`backends` | [BackendConfig](#backendconfig)`[]`<br>List of storage backends
+`namespaces` | [NamespaceConfig](#namespaceconfig)`[]`<br>List of namespaces for storage
+
+## BackendConfig
+
+Common settings object for different providers
+
+Name | Description
+---- | -----------
+`id` | `string`<br>The storage backend ID
+`base` | `string`<br>Base backend to copy default settings from
+`secondary` | `string`<br>Specifies another backend to read from if an object is not found in this one. Can be used when migrating data from one backend to another.
+`type` | [StorageBackendType](#storagebackendtype-enum)<br>
+`baseDir` | `string`<br>
+`awsBucketName` | `string`<br>Name of the bucket to use
+`awsBucketPath` | `string`<br>Base path within the bucket
+`awsCredentials` | [AwsCredentialsType](#awscredentialstype-enum)<br>Type of credentials to use
+`awsRole` | `string`<br>ARN of a role to assume
+`awsProfile` | `string`<br>The AWS profile to read credentials form
+`awsRegion` | `string`<br>Region to connect to
+`azureConnectionString` | `string`<br>Connection string for Azure
+`azureContainerName` | `string`<br>Name of the container
+`relayServer` | `string`<br>
+`relayToken` | `string`<br>
+`gcsBucketName` | `string`<br>Name of the GCS bucket to use
+`gcsBucketPath` | `string`<br>Base path within the bucket
+
+## StorageBackendType (Enum)
+
+Types of storage backend to use
+
+Name | Description
+---- | -----------
+`FileSystem` | Local filesystem
+`Aws` | AWS S3
+`Azure` | Azure blob store
+`Gcs` | Google Cloud Storage
+`Memory` | In-memory only (for testing)
+
+## AwsCredentialsType (Enum)
+
+Credentials to use for AWS
+
+Name | Description
+---- | -----------
+`Default` | Use default credentials from the AWS SDK
+`Profile` | Read credentials from the  profile in the AWS config file
+`AssumeRole` | Assume a particular role. Should specify ARN in
+`AssumeRoleWebIdentity` | Assume a particular role using the current environment variables.
+
+## NamespaceConfig
+
+Configuration of a particular namespace
+
+Name | Description
+---- | -----------
+`id` | `string`<br>Identifier for this namespace
+`backend` | `string`<br>Backend to use for this namespace
+`prefix` | `string`<br>Prefix for items within this namespace
+`gcFrequencyHrs` | `number`<br>How frequently to run garbage collection, in hours.
+`gcDelayHrs` | `number`<br>How long to keep newly uploaded orphanned objects before allowing them to be deleted, in hours.
+`enableAliases` | `boolean`<br>Support querying exports by their aliases
+`acl` | [AclConfig](#aclconfig)<br>Access list for this namespace
+
 ## ComputeConfig
 
@@ -472,96 +485,6 @@
 `description` | `string`<br>Human-readable description
 `computeId` | `string`<br>Compute ID for this network (used when allocating compute resources)
 
-## ExperimentalConfig
-
-Configuration for the experimental plugin
-
-Name | Description
----- | -----------
-`notifications` | [NotificationConfig](Notification.md)`[]`<br>Notifications to be sent on the Horde server
-
-## SecretsConfig
-
-Configuration for the secrets system
-
-Name | Description
----- | -----------
-`secrets` | [SecretConfig](#secretconfig)`[]`<br>List of secrets
-`providerConfigs` | [SecretProviderConfig](#secretproviderconfig)`[]`<br>List of configurations for secret providers
-
-## SecretConfig
-
-Configuration for a secret value
-
-Name | Description
----- | -----------
-`id` | `string`<br>Identifier for this secret
-`data` | `string` `->` `string`<br>Key/value pairs associated with this secret
-`sources` | [ExternalSecretConfig](#externalsecretconfig)`[]`<br>Providers to source key/value pairs from
-`acl` | [AclConfig](#aclconfig)<br>Defines access to this particular secret
-
-## ExternalSecretConfig
-
-Configuration for an external secret provider
-
-Name | Description
----- | -----------
-`providerConfig` | `string`<br>The name to look up a secret provider config. A config is mutually exclusive to and to be used when a secret requires additional options than just the name of the provider
-`provider` | `string`<br>Name of the provider to use
-`format` | [ExternalSecretFormat](#externalsecretformat-enum)<br>Format of the secret
-`key` | `string`<br>Optional key indicating the parameter to set in the resulting data array. Required if  is .
-`path` | `string`<br>Optional value indicating what to fetch from the provider
-
-## ExternalSecretFormat (Enum)
-
-Format describing how to parse external secret values
-
-Name | Description
----- | -----------
-`Text` | Secret is a plain text value which will be stored using the external secret key
-`Json` | Secret is a JSON formatted string containing key/value pairs
-
-## SecretProviderConfig
-
-Common settings object for different secret providers
-
-Name | Description
----- | -----------
-`name` | `string`<br>Name of the secret provider config
-`provider` | `string`<br>Name of the provider
-`hcpVault` | [HcpVaultConfig](#hcpvaultconfig)<br>Configuration for HCP Vault
-`azureKeyVault` | [AzureKeyVaultConfig](#azurekeyvaultconfig)<br>Configuration for Azure Key Vault
-
-## HcpVaultConfig
-
-Configuration for HCP Vault secrets.
-
-Name | Description
----- | -----------
-`preSharedKey` | `string`<br>The Vault authentication token to use when the type of credentials is .<br>More information on this token can be found in the API documentation for X-Vault-Token
-`endPoint` | `string`<br>The address of the Vault server.
-`awsArnRole` | `string`<br>The optional AWS ARN to assume a role when logging into Vault. To be used when the type of credentials is . When not set the AWS SDK default credential search will be used.
-`awsIamServerId` | `string`<br>The value for the X-Vault-AWS-IAM-Server-ID header when the Vault server is configured with iam_server_id_header_value. Most likely to be the Vault server's DNS name. To be used when the type of credentials is .
-`role` | `string`<br>The name of the role in the Vault configuration for the login. This role name can be different to the name of the AWS IAM principal. To be used when the type of credentials is . When not set the name of IAM principal will be used.
-`credentials` | [HcpVaultCredentialsType](#hcpvaultcredentialstype-enum)<br>Type of credentials to use.
-
-## HcpVaultCredentialsType (Enum)
-
-Credentials to use to log into HCP Vault.
-
-Name | Description
----- | -----------
-`PreSharedKey` | Use a pre-shared token provided by Vault for authentication.
-`AwsAuth` | Use AWS IAM credentials to authenticate against the Vault server.
-
-## AzureKeyVaultConfig
-
-Configuration for Azure Key Vault secrets.
-
-Name | Description
----- | -----------
-`vaultUri` | `string`<br>The Uri to the Azure Key Vault.
-
 ## SymbolsConfig
 
 Configuration for the tools system
@@ -581,83 +504,152 @@
 `public` | `boolean`<br>Whether to make this store available without auth
 `acl` | [AclConfig](#aclconfig)<br>Access to the symbol store
 
-## StorageConfig
+## BuildConfig
 
-Configuration for storage
+Configuration for the build plugin
 
 Name | Description
 ---- | -----------
-`enableGc` | `boolean`<br>Whether to enable garbage collection
-`enableGcVerification` | `boolean`<br>Whether to enable garbage collection in verification mode (nothing deleted, just logging on access to deleted blobs)
-`gcWorkerCount` | `integer`<br>Number of concurrent GC reachability-check workers. Default 8. Reachability checks are I/O-bound (MongoDB queries), so higher values scale throughput roughly linearly up to the database's capacity.
-`gcQueueLimit` | `integer`<br>Maximum number of entries in a namespace's Redis GC queue before the blob scanner pauses. Default 50,000.
-`enableReliableGcEnqueue` | `boolean`<br>Whether to await Redis writes when enqueuing blobs for GC checking. When false (default), uses fire-and-forget which can silently drop entries on Redis failures. Set true for reliable enqueue.
-`enablePerNamespaceQueueIsolation` | `boolean`<br>Whether to isolate GC queue limits per-namespace. When false (default), one backed-up namespace pauses the scanner for ALL namespaces. When true, only the backed-up namespace is paused and its blobs are deferred.
-`gcBlobIngestionBatchSize` | `integer`<br>Number of blobs to process per blob-ingestion tick (TickBlobsAsync). Default 500. Raise to increase scan throughput when catching up on a backlog of unscanned blobs.
-`gcBatchSize` | `integer`<br>Number of blob IDs to batch together for the import-reachability check. Default 1 (current per-blob behavior). When greater than 1, uses $in query to check multiple blobs in a single MongoDB round-trip. Values of 50-200 recommended.
-`backends` | [BackendConfig](#backendconfig)`[]`<br>List of storage backends
-`namespaces` | [NamespaceConfig](#namespaceconfig)`[]`<br>List of namespaces for storage
+`perforceClusters` | [PerforceCluster](#perforcecluster)`[]`<br>List of Perforce clusters
+`devices` | [DeviceConfig](#deviceconfig)<br>Device configuration
+`maxConformCount` | `integer`<br>Maximum number of conforms to run at once
+`agentShutdownIfDisabledGracePeriod` | `string`<br>Time to wait before shutting down an agent that has been disabled Used if no value is set on the actual pool.
+`robomerge` | [RoboMergeConfig](#robomergeconfig)<br>RoboMerge branchmap file configuration for the merge graph service. When null or empty, the service is inert (no Perforce connections, no API data).
+`artifactTypes` | [ArtifactTypeConfig](#artifacttypeconfig)`[]`<br>Configuration for different artifact types
+`projects` | [ProjectConfig](Projects.md)`[]`<br>List of projects
+`enableConformTasks` | `boolean`<br>Whether to allow conform tasks to run
+`enableJobTaskConformCheck` | `boolean`<br>Whether to check if a conform is required on an agent before assigning a lease for a job task source.
+`issueFixedTag` | `string`<br>Commit tag to use for marking issues as fixed
+`issues` | [IssueConfig](Issues.md)<br>Build health issue configurations
+`enableIssueRootCauseAnalysisNotifications` | `boolean`<br>Whether to allow notifications for Root Cause Analysis to be sent
+`enableAutoRcaSlackNotifications` | `boolean`<br>Global kill switch for posting auto-RCA results to Slack triage threads. When false (default), no auto-RCA Slack notifications are sent regardless of per-workflow AutoRcaSlackEnabled settings.
+`autoRcaUseMultiStage` | `boolean`<br>Whether to use multi-stage LLM analysis for auto-RCA. Multi-stage is slower but potentially more accurate. Set to false for faster processing with higher throughput.
+`maxConcurrentRcaRequests` | `integer`<br>Maximum number of concurrent auto-RCA requests. Limits load on the external LLM API. Set to 0 to disable auto-RCA globally (acts as a kill switch).
+`issueFindSuspectsMaxSuccessDays` | `integer`<br>The maximum days since node success to scan for suspects
+`issueFindSuspectsMaxScoringSeconds` | `integer`<br>The maximum number of seconds to spend scording suspects
+`issueFindSuspectsMaxChanges` | `integer`<br>Maximum number of commits to examine for suspects. Previously hardcoded as const MaxChanges = 500.
+`issueFindSuspectsEnableParallelScoring` | `boolean`<br>Enable parallel scoring of commits for suspect finding. When enabled, commits are scored concurrently up to MaxParallelism.
+`issueFindSuspectsMaxParallelism` | `integer`<br>Maximum concurrent commit file lookups when parallel scoring is enabled. Higher values improve speed but increase Perforce server load.
+`issueFindSuspectsEnableEarlyExit` | `boolean`<br>Enable early exit when theoretical maximum score is achieved. Only exits early if the best possible score has been found.
+`issueFindSuspectsEnableBidirectionalQuery` | `boolean`<br>Enable bidirectional commit query to capture both recent and old commits. Queries half the limit from each end of the commit range.
+`issueFindSuspectsEnableFingerprintFallback` | `boolean`<br>Enable fingerprint-based fallback when no suspects are found. Helps prevent duplicate issues when commit range exceeds MaxChanges.
+`issueFingerprintFallbackMinScore` | `integer`<br>Minimum score required for fingerprint fallback matching. Higher values are more conservative (fewer false matches).
+`issueFingerprintFallbackMaxAgeDays` | `integer`<br>Maximum age in days for issues to be considered in fingerprint fallback.
+`issueFingerprintFallbackExcludedTypes` | `string[]`<br>Fingerprint types excluded from fallback matching (too generic).
+`enableExpandedIgnorePatterns` | `boolean`<br>Whether to enable expanded IgnorePatterns discovery (workspace-wide scanning + config-specified files). When false (default), only Engine directory is scanned for ignore patterns.
+`lowContextFingerprintTypes` | `string[]`<br>Fingerprint type prefixes considered "low context" for suspect mention suppression. When an issue's fingerprint type starts with any of these prefixes AND all suspects are tied at rank 0, @ mentions are suppressed in Slack. Suspects are still listed without mentions. Set to an empty array to disable suppression entirely.
 
-## BackendConfig
+## PerforceCluster
 
-Common settings object for different providers
+Information about a cluster of Perforce servers.
 
 Name | Description
 ---- | -----------
-`id` | `string`<br>The storage backend ID
-`base` | `string`<br>Base backend to copy default settings from
-`secondary` | `string`<br>Specifies another backend to read from if an object is not found in this one. Can be used when migrating data from one backend to another.
-`type` | [StorageBackendType](#storagebackendtype-enum)<br>
-`baseDir` | `string`<br>
-`awsBucketName` | `string`<br>Name of the bucket to use
-`awsBucketPath` | `string`<br>Base path within the bucket
-`awsCredentials` | [AwsCredentialsType](#awscredentialstype-enum)<br>Type of credentials to use
-`awsRole` | `string`<br>ARN of a role to assume
-`awsProfile` | `string`<br>The AWS profile to read credentials form
-`awsRegion` | `string`<br>Region to connect to
-`azureConnectionString` | `string`<br>Connection string for Azure
-`azureContainerName` | `string`<br>Name of the container
-`relayServer` | `string`<br>
-`relayToken` | `string`<br>
-`gcsBucketName` | `string`<br>Name of the GCS bucket to use
-`gcsBucketPath` | `string`<br>Base path within the bucket
+`name` | `string`<br>Name of the cluster
+`serviceAccount` | `string`<br>Username for Horde to log in to this server. Will use the first account specified below if not overridden.
+`canImpersonate` | `boolean`<br>Whether the service account can impersonate other users
+`supportsPartitionedWorkspaces` | `boolean`<br>Whether to use partitioned workspaces on this server
+`servers` | [PerforceServer](#perforceserver)`[]`<br>List of servers
+`credentials` | [PerforceCredentials](#perforcecredentials)`[]`<br>List of server credentials
+`autoSdk` | [AutoSdkWorkspace](#autosdkworkspace)`[]`<br>List of autosdk streams
 
-## StorageBackendType (Enum)
+## PerforceServer
 
-Types of storage backend to use
+Information about an individual Perforce server
 
 Name | Description
 ---- | -----------
-`FileSystem` | Local filesystem
-`Aws` | AWS S3
-`Azure` | Azure blob store
-`Gcs` | Google Cloud Storage
-`Memory` | In-memory only (for testing)
+`serverAndPort` | `string`<br>The server and port. The server may be a DNS entry with multiple records, in which case it will be actively load balanced. If "ssl:" prefix is used, ensure P4 server's fingerprint/certificate is trusted. See Horde's documentation on connecting to SSL-enabled Perforce servers.
+`healthCheck` | `boolean`<br>Whether to query the healthcheck address under each server
+`resolveDns` | `boolean`<br>Whether to resolve the DNS entries and load balance between different hosts
+`maxConformCount` | `integer`<br>Maximum number of simultaneous conforms on this server
+`condition` | `string`<br>Optional condition for a machine to be eligible to use this server
+`properties` | `string[]`<br>List of properties for an agent to be eligible to use this server
 
-## AwsCredentialsType (Enum)
+## PerforceCredentials
 
-Credentials to use for AWS
+Credentials for a Perforce user
 
 Name | Description
 ---- | -----------
-`Default` | Use default credentials from the AWS SDK
-`Profile` | Read credentials from the  profile in the AWS config file
-`AssumeRole` | Assume a particular role. Should specify ARN in
-`AssumeRoleWebIdentity` | Assume a particular role using the current environment variables.
+`userName` | `string`<br>The username
+`password` | `string`<br>Password for the user
+`ticket` | `string`<br>Login ticket for the user (will be used instead of password if set)
 
-## NamespaceConfig
+## AutoSdkWorkspace
 
-Configuration of a particular namespace
+Path to a platform and stream to use for syncing AutoSDK
 
 Name | Description
 ---- | -----------
-`id` | `string`<br>Identifier for this namespace
-`backend` | `string`<br>Backend to use for this namespace
-`prefix` | `string`<br>Prefix for items within this namespace
-`gcFrequencyHrs` | `number`<br>How frequently to run garbage collection, in hours.
-`gcDelayHrs` | `number`<br>How long to keep newly uploaded orphanned objects before allowing them to be deleted, in hours.
-`enableAliases` | `boolean`<br>Support querying exports by their aliases
-`acl` | [AclConfig](#aclconfig)<br>Access list for this namespace
+`name` | `string`<br>Name of this workspace
+`properties` | `string[]`<br>The agent properties to check (eg. "OSFamily=Windows")
+`userName` | `string`<br>Username for logging in to the server
+`stream` | `string`<br>Stream to use
 
+## DeviceConfig
+
+Configuration for devices
+
+Name | Description
+---- | -----------
+`platforms` | [DevicePlatformConfig](#deviceplatformconfig)`[]`<br>List of device platforms
+`pools` | [DevicePoolConfig](#devicepoolconfig)`[]`<br>List of device pools
+
+## DevicePlatformConfig
+
+Configuration for a device platform
+
+Name | Description
+---- | -----------
+`id` | `string`<br>The id for this platform
+`name` | `string`<br>Name of the platform
+`models` | `string[]`<br>A list of platform models
+`legacyNames` | `string[]`<br>Legacy names which older versions of Gauntlet may be using
+`legacyPerfSpecHighModel` | `string`<br>Model name for the high perf spec, which may be requested by Gauntlet
+
+## DevicePoolConfig
+
+Configuration for a device pool
+
+Name | Description
+---- | -----------
+`id` | `string`<br>The id for this platform
+`name` | `string`<br>The name of the pool
+`poolType` | [DevicePoolType](#devicepooltype-enum)<br>The type of the pool
+`projectIds` | `string[]`<br>List of project ids associated with pool
+
+## DevicePoolType (Enum)
+
+The type of device pool
+
+Name | Description
+---- | -----------
+`Automation` | Available to CIS jobs
+`Shared` | Shared by users with remote checking and checkouts
+
+## RoboMergeConfig
+
+Configuration for the RoboMerge branchmap service
+
+Name | Description
+---- | -----------
+`branchmapFiles` | `string[]`<br>Perforce depot paths to branchmap JSON files. Plain depot paths (e.g. "//MyProject/Main/Build/Streams/myproject.branchmap.json"). All files are read from the Perforce cluster specified by .
+`perforceCluster` | `string`<br>Perforce cluster name to connect to for reading branchmap files. Maps to a connection configured in BuildServerConfig.Perforce[]. Default: "default" (uses the default Perforce connection).
+`tickInterval` | `string`<br>How often to poll Perforce for branchmap file changes. Default: 5 minutes. Changing this requires a server restart (ticker interval is set at construction time).
+
+## ArtifactTypeConfig
+
+Configuration for an artifact
+
+Name | Description
+---- | -----------
+`name` | `string`<br>Legacy 'Name' property
+`type` | `string`<br>Name of the artifact type
+`acl` | [AclConfig](#aclconfig)<br>Acl for the artifact type
+`keepCount` | `integer`<br>Number of artifacts to retain
+`keepDays` | `integer`<br>Number of days to retain artifacts of this type
+`namespaceId` | `string`<br>Storage namespace to use for this artifact types
+
 ## ToolsConfig
 
@@ -687,3 +679,11 @@
 `namespaceId` | `string`<br>Default namespace for new deployments of this tool
 `maxDeployments` | `integer`<br>Maximum number of deployments to retain for this tool. When exceeded, the oldest deployments are removed.
 `acl` | [AclConfig](#aclconfig)<br>Permissions for the tool
+
+## ExperimentalConfig
+
+Configuration for the experimental plugin
+
+Name | Description
+---- | -----------
+`notifications` | [NotificationConfig](Notification.md)`[]`<br>Notifications to be sent on the Horde server
diff --git a/Engine/Source/Programs/Horde/Docs/Deployment/ServerSettings.md b/Engine/Source/Programs/Horde/Docs/Deployment/ServerSettings.md
--- a/Engine/Source/Programs/Horde/Docs/Deployment/ServerSettings.md
+++ b/Engine/Source/Programs/Horde/Docs/Deployment/ServerSettings.md
@@ -128,14 +128,14 @@
 Name | Description
 ---- | -----------
 `analytics` | [AnalyticsServerConfig](#analyticsserverconfig)<br>Configuration for the analytics plugin
-`build` | [BuildServerConfig](#buildserverconfig)<br>Configuration for the build plugin
 `exampleplugin` | [ExamplePluginServerConfig](#examplepluginserverconfig)<br>Configuration for the exampleplugin plugin
-`compute` | [ComputeServerConfig](#computeserverconfig)<br>Configuration for the compute plugin
-`experimental` | [ExperimentalServerConfig](#experimentalserverconfig)<br>Configuration for the experimental plugin
 `secrets` | [SecretsServerConfig](#secretsserverconfig)<br>Configuration for the secrets plugin
-`symbols` | [PluginServerConfig](#pluginserverconfig)<br>Configuration for the symbols plugin
 `storage` | [StorageServerConfig](#storageserverconfig)<br>Configuration for the storage plugin
+`compute` | [ComputeServerConfig](#computeserverconfig)<br>Configuration for the compute plugin
+`symbols` | [PluginServerConfig](#pluginserverconfig)<br>Configuration for the symbols plugin
+`build` | [BuildServerConfig](#buildserverconfig)<br>Configuration for the build plugin
 `tools` | [ToolsServerConfig](#toolsserverconfig)<br>Configuration for the tools plugin
+`experimental` | [ExperimentalServerConfig](#experimentalserverconfig)<br>Configuration for the experimental plugin
 
 ## AnalyticsServerConfig
 
@@ -174,6 +174,135 @@
 `retainDays` | `number`<br>Number of days worth of telmetry events to keep
 `enabled` | `boolean`<br>Whether to enable this sink
 
+## ExamplePluginServerConfig
+
+Server configuration for the ExamplePlugin
+
+Name | Description
+---- | -----------
+`enabled` | `boolean`<br>Whether the plugin should be enabled or not
+
+## SecretsServerConfig
+
+Static configuration for the secrets plugin
+
+Name | Description
+---- | -----------
+`withAws` | `boolean`<br>Whether to enable Amazon Web Services (AWS) specific features
+`enabled` | `boolean`<br>Whether the plugin should be enabled or not
+
+## StorageServerConfig
+
+Static settings for the storage system
+
+Name | Description
+---- | -----------
+`bundleCacheDir` | `string`<br>Directory to use for the coarse-grained backend cache. This caches full bundles downloaded from the upstream object store.
+`bundleCacheSize` | `string`<br>Maximum size of the storage cache on disk. Accepts standard binary suffixes (kb, mb, gb, tb, etc...)
+`bundleCacheSizeBytes` | `integer`<br>Accessor for the bundle cache size in bytes
+`backends` | [BackendConfig](#backendconfig)`[]`<br>Overridden settings for storage backends. Useful for running against a production server with custom backends.
+`enabled` | `boolean`<br>Whether the plugin should be enabled or not
+
+## BackendConfig
+
+Common settings object for different providers
+
+Name | Description
+---- | -----------
+`id` | `string`<br>The storage backend ID
+`base` | `string`<br>Base backend to copy default settings from
+`secondary` | `string`<br>Specifies another backend to read from if an object is not found in this one. Can be used when migrating data from one backend to another.
+`type` | [StorageBackendType](#storagebackendtype-enum)<br>
+`baseDir` | `string`<br>
+`awsBucketName` | `string`<br>Name of the bucket to use
+`awsBucketPath` | `string`<br>Base path within the bucket
+`awsCredentials` | [AwsCredentialsType](#awscredentialstype-enum)<br>Type of credentials to use
+`awsRole` | `string`<br>ARN of a role to assume
+`awsProfile` | `string`<br>The AWS profile to read credentials form
+`awsRegion` | `string`<br>Region to connect to
+`azureConnectionString` | `string`<br>Connection string for Azure
+`azureContainerName` | `string`<br>Name of the container
+`relayServer` | `string`<br>
+`relayToken` | `string`<br>
+`gcsBucketName` | `string`<br>Name of the GCS bucket to use
+`gcsBucketPath` | `string`<br>Base path within the bucket
+
+## StorageBackendType (Enum)
+
+Types of storage backend to use
+
+Name | Description
+---- | -----------
+`FileSystem` | Local filesystem
+`Aws` | AWS S3
+`Azure` | Azure blob store
+`Gcs` | Google Cloud Storage
+`Memory` | In-memory only (for testing)
+
+## AwsCredentialsType (Enum)
+
+Credentials to use for AWS
+
+Name | Description
+---- | -----------
+`Default` | Use default credentials from the AWS SDK
+`Profile` | Read credentials from the  profile in the AWS config file
+`AssumeRole` | Assume a particular role. Should specify ARN in
+`AssumeRoleWebIdentity` | Assume a particular role using the current environment variables.
+
+## ComputeServerConfig
+
+Static configuration for the compute plugin
+
+Name | Description
+---- | -----------
+`enableUpgradeTasks` | `boolean`<br>Whether to enable the upgrade task source, always upgrading agents to the latest version
+`withAws` | `boolean`<br>Whether to enable Amazon Web Services (AWS) specific features
+`awsRegions` | `string[]`<br>List of AWS regions for Horde to be aware of (e.g. us-east-1 or eu-central-1) Right now, this is only used for replicating CloudWatch metrics to multiple regions
+`awsAutoScalingQueueUrls` | `string[]`<br>AWS SQS queue URLs where lifecycle events from EC2 auto-scaling are received
+`fleetManagerV2` | [FleetManagerType](#fleetmanagertype-enum)<br>Default fleet manager to use (when not specified by pool)
+`fleetManagerV2Config` | `object`<br>Config for the fleet manager (serialized JSON)
+`autoEnrollAgents` | `boolean`<br>Whether to automatically enroll agents in the farm
+`defaultAgentPoolSizeStrategy` | [PoolSizeStrategy](#poolsizestrategy-enum)<br>Default agent pool sizing strategy for pools that doesn't have one explicitly configured
+`agentPoolScaleOutCooldownSeconds` | `integer`<br>Scale-out cooldown for auto-scaling agent pools (in seconds). Can be overridden by per-pool settings.
+`agentPoolScaleInCooldownSeconds` | `integer`<br>Scale-in cooldown for auto-scaling agent pools (in seconds). Can be overridden by per-pool settings.
+`computeTunnelPort` | `integer`<br>Port to listen on for tunneling compute sockets to agents
+`computeTunnelAddress` | `string`<br>What address (host:port) clients should connect to for compute socket tunneling Port may differ from  if Horde server is behind a reverse proxy/firewall
+`enabled` | `boolean`<br>Whether the plugin should be enabled or not
+
+## FleetManagerType (Enum)
+
+Available fleet managers
+
+Name | Description
+---- | -----------
+`Default` | Default fleet manager
+`NoOp` | No-op fleet manager.
+`Aws` | Fleet manager for handling AWS EC2 instances. Will create and/or terminate instances from scratch.
+`AwsReuse` | Fleet manager for handling AWS EC2 instances. Will start already existing but stopped instances to reuse existing EBS disks.
+`AwsRecycle` | Fleet manager for handling AWS EC2 instances. Will start already existing but stopped instances to reuse existing EBS disks.
+`AwsAsg` | Fleet manager for handling AWS EC2 instances. Uses an EC2 auto-scaling group for controlling the number of running instances.
+
+## PoolSizeStrategy (Enum)
+
+Available pool sizing strategies
+
+Name | Description
+---- | -----------
+`LeaseUtilization` | Strategy based on lease utilization
+`JobQueue` | Strategy based on size of job build queue
+`NoOp` | No-op strategy used as fallback/default behavior
+`ComputeQueueAwsMetric` | A no-op strategy that reports metrics to let an external AWS auto-scaling policy scale the fleet
+`LeaseUtilizationAwsMetric` | A no-op strategy that reports metrics to let an external AWS auto-scaling policy scale the fleet
+
+## PluginServerConfig
+
+Base class for plugin server config objects
+
+Name | Description
+---- | -----------
+`enabled` | `boolean`<br>Whether the plugin should be enabled or not
+
 ## BuildServerConfig
 
 Static configuration for the build plugin
@@ -319,143 +448,6 @@
 `maxChildCount` | `integer`<br>Maximum number of children in each node
 `sliceThreshold` | `integer`<br>Threshold hash value for splitting interior nodes
 
-## ExamplePluginServerConfig
-
-Server configuration for the ExamplePlugin
-
-Name | Description
----- | -----------
-`enabled` | `boolean`<br>Whether the plugin should be enabled or not
-
-## ComputeServerConfig
-
-Static configuration for the compute plugin
-
-Name | Description
----- | -----------
-`enableUpgradeTasks` | `boolean`<br>Whether to enable the upgrade task source, always upgrading agents to the latest version
-`withAws` | `boolean`<br>Whether to enable Amazon Web Services (AWS) specific features
-`awsRegions` | `string[]`<br>List of AWS regions for Horde to be aware of (e.g. us-east-1 or eu-central-1) Right now, this is only used for replicating CloudWatch metrics to multiple regions
-`awsAutoScalingQueueUrls` | `string[]`<br>AWS SQS queue URLs where lifecycle events from EC2 auto-scaling are received
-`fleetManagerV2` | [FleetManagerType](#fleetmanagertype-enum)<br>Default fleet manager to use (when not specified by pool)
-`fleetManagerV2Config` | `object`<br>Config for the fleet manager (serialized JSON)
-`autoEnrollAgents` | `boolean`<br>Whether to automatically enroll agents in the farm
-`defaultAgentPoolSizeStrategy` | [PoolSizeStrategy](#poolsizestrategy-enum)<br>Default agent pool sizing strategy for pools that doesn't have one explicitly configured
-`agentPoolScaleOutCooldownSeconds` | `integer`<br>Scale-out cooldown for auto-scaling agent pools (in seconds). Can be overridden by per-pool settings.
-`agentPoolScaleInCooldownSeconds` | `integer`<br>Scale-in cooldown for auto-scaling agent pools (in seconds). Can be overridden by per-pool settings.
-`computeTunnelPort` | `integer`<br>Port to listen on for tunneling compute sockets to agents
-`computeTunnelAddress` | `string`<br>What address (host:port) clients should connect to for compute socket tunneling Port may differ from  if Horde server is behind a reverse proxy/firewall
-`enabled` | `boolean`<br>Whether the plugin should be enabled or not
-
-## FleetManagerType (Enum)
-
-Available fleet managers
-
-Name | Description
----- | -----------
-`Default` | Default fleet manager
-`NoOp` | No-op fleet manager.
-`Aws` | Fleet manager for handling AWS EC2 instances. Will create and/or terminate instances from scratch.
-`AwsReuse` | Fleet manager for handling AWS EC2 instances. Will start already existing but stopped instances to reuse existing EBS disks.
-`AwsRecycle` | Fleet manager for handling AWS EC2 instances. Will start already existing but stopped instances to reuse existing EBS disks.
-`AwsAsg` | Fleet manager for handling AWS EC2 instances. Uses an EC2 auto-scaling group for controlling the number of running instances.
-
-## PoolSizeStrategy (Enum)
-
-Available pool sizing strategies
-
-Name | Description
----- | -----------
-`LeaseUtilization` | Strategy based on lease utilization
-`JobQueue` | Strategy based on size of job build queue
-`NoOp` | No-op strategy used as fallback/default behavior
-`ComputeQueueAwsMetric` | A no-op strategy that reports metrics to let an external AWS auto-scaling policy scale the fleet
-`LeaseUtilizationAwsMetric` | A no-op strategy that reports metrics to let an external AWS auto-scaling policy scale the fleet
-
-## ExperimentalServerConfig
-
-Server configuration for the experimental plugin
-
-Name | Description
----- | -----------
-`enabled` | `boolean`<br>Whether the plugin should be enabled or not
-
-## SecretsServerConfig
-
-Static configuration for the secrets plugin
-
-Name | Description
----- | -----------
-`withAws` | `boolean`<br>Whether to enable Amazon Web Services (AWS) specific features
-`enabled` | `boolean`<br>Whether the plugin should be enabled or not
-
-## PluginServerConfig
-
-Base class for plugin server config objects
-
-Name | Description
----- | -----------
-`enabled` | `boolean`<br>Whether the plugin should be enabled or not
-
-## StorageServerConfig
-
-Static settings for the storage system
-
-Name | Description
----- | -----------
-`bundleCacheDir` | `string`<br>Directory to use for the coarse-grained backend cache. This caches full bundles downloaded from the upstream object store.
-`bundleCacheSize` | `string`<br>Maximum size of the storage cache on disk. Accepts standard binary suffixes (kb, mb, gb, tb, etc...)
-`bundleCacheSizeBytes` | `integer`<br>Accessor for the bundle cache size in bytes
-`backends` | [BackendConfig](#backendconfig)`[]`<br>Overridden settings for storage backends. Useful for running against a production server with custom backends.
-`enabled` | `boolean`<br>Whether the plugin should be enabled or not
-
-## BackendConfig
-
-Common settings object for different providers
-
-Name | Description
----- | -----------
-`id` | `string`<br>The storage backend ID
-`base` | `string`<br>Base backend to copy default settings from
-`secondary` | `string`<br>Specifies another backend to read from if an object is not found in this one. Can be used when migrating data from one backend to another.
-`type` | [StorageBackendType](#storagebackendtype-enum)<br>
-`baseDir` | `string`<br>
-`awsBucketName` | `string`<br>Name of the bucket to use
-`awsBucketPath` | `string`<br>Base path within the bucket
-`awsCredentials` | [AwsCredentialsType](#awscredentialstype-enum)<br>Type of credentials to use
-`awsRole` | `string`<br>ARN of a role to assume
-`awsProfile` | `string`<br>The AWS profile to read credentials form
-`awsRegion` | `string`<br>Region to connect to
-`azureConnectionString` | `string`<br>Connection string for Azure
-`azureContainerName` | `string`<br>Name of the container
-`relayServer` | `string`<br>
-`relayToken` | `string`<br>
-`gcsBucketName` | `string`<br>Name of the GCS bucket to use
-`gcsBucketPath` | `string`<br>Base path within the bucket
-
-## StorageBackendType (Enum)
-
-Types of storage backend to use
-
-Name | Description
----- | -----------
-`FileSystem` | Local filesystem
-`Aws` | AWS S3
-`Azure` | Azure blob store
-`Gcs` | Google Cloud Storage
-`Memory` | In-memory only (for testing)
-
-## AwsCredentialsType (Enum)
-
-Credentials to use for AWS
-
-Name | Description
----- | -----------
-`Default` | Use default credentials from the AWS SDK
-`Profile` | Read credentials from the  profile in the AWS config file
-`AssumeRole` | Assume a particular role. Should specify ARN in
-`AssumeRoleWebIdentity` | Assume a particular role using the current environment variables.
-
 ## ToolsServerConfig
 
 Server configuration for bundled tools
@@ -529,3 +521,11 @@
 `actions` | `string[]`<br>Actions to include
 `excludeActions` | `string[]`<br>Actions to exclude from the inherited actions
 `extends` | `string[]`<br>Other profiles to extend from
+
+## ExperimentalServerConfig
+
+Server configuration for the experimental plugin
+
+Name | Description
+---- | -----------
+`enabled` | `boolean`<br>Whether the plugin should be enabled or not
