Horde: Fix vertical centering of compact label badges in job view
Smaller non-error labels were not vertically centered against
larger error labels in job view rows. Set lineHeight:0 on the
LabelBadge wrapper to collapse extra inline space from nested
elements, allowing flex alignItems:center to work correctly.

diff --git a/Engine/Source/Programs/Horde/HordeDashboard/src/components/labels/LabelBadge.tsx b/Engine/Source/Programs/Horde/HordeDashboard/src/components/labels/LabelBadge.tsx
--- a/Engine/Source/Programs/Horde/HordeDashboard/src/components/labels/LabelBadge.tsx
+++ b/Engine/Source/Programs/Horde/HordeDashboard/src/components/labels/LabelBadge.tsx
@@ -27,7 +27,7 @@ export const LabelBadge: React.FC<LabelBadgeProps> = ({ label, jobId, labelIndex
     const url = labelIndex >= 0 ? `/job/${jobId}?label=${labelIndex}` : `/job/${jobId}`;
 
     const badge = (
-        <div id={`label_${jobId}_${label.dashboardName}_${label.dashboardCategory}`.replace(/[^A-Za-z0-9]/g, "")} style={{ display: "inline-block" }}>
+        <div id={`label_${jobId}_${label.dashboardName}_${label.dashboardCategory}`.replace(/[^A-Za-z0-9]/g, "")} style={{ lineHeight: 0 }}>
             <Link to={url}><div className={badgeClass}>
                 <DefaultButton key={label.dashboardName} style={{ backgroundColor: color.primaryColor }} text={label.dashboardName ?? 'Other'}
                     onMouseOver={(ev) => ev.stopPropagation()}
