[Horde Dashboard] Remove redundant /api/v1/agents call from Pool page

Before: PoolAgentPanelcalls agentStore.update() which updates the full global store.

After: Externally pass agents in so we only make the singular less expensive scoped /api/v1/agents call

diff --git a/Engine/Source/Programs/Horde/HordeDashboard/src/components/PoolView.tsx b/Engine/Source/Programs/Horde/HordeDashboard/src/components/PoolView.tsx
--- a/Engine/Source/Programs/Horde/HordeDashboard/src/components/PoolView.tsx
+++ b/Engine/Source/Programs/Horde/HordeDashboard/src/components/PoolView.tsx
@@ -290,7 +290,7 @@ const PoolAgentPanel: React.FC<{ poolId: string }> = ({ poolId }) => {
          <Stack>
             <Text variant="mediumPlus" styles={{ root: { fontFamily: "Horde Open Sans SemiBold" } }}>Agents</Text>
          </Stack>
-         <AgentPanel poolId={poolId} />
+         <AgentPanel poolId={poolId} agents={handler.agents} />
       </Stack>
    </Stack>
 }
diff --git a/Engine/Source/Programs/Horde/HordeDashboard/src/components/agents/AgentStore.ts b/Engine/Source/Programs/Horde/HordeDashboard/src/components/agents/AgentStore.ts
--- a/Engine/Source/Programs/Horde/HordeDashboard/src/components/agents/AgentStore.ts
+++ b/Engine/Source/Programs/Horde/HordeDashboard/src/components/agents/AgentStore.ts
@@ -118,10 +118,15 @@ class AgentStore {
 
     }
 
+    async updatePools(): Promise<void> {
+        const pools = await backend.getPools();
+        this._setPools(pools);
+    }
+
     async update(slim:boolean = false, invalidateCache: boolean = false): Promise<void> {
         return new Promise<void>((resolve, reject) => {
 
-            if (this.inflight !== undefined && slim) {                
+            if (this.inflight !== undefined && slim) {
                 resolve();
                 return;
             }
diff --git a/Engine/Source/Programs/Horde/HordeDashboard/src/components/agents/AgentView.tsx b/Engine/Source/Programs/Horde/HordeDashboard/src/components/agents/AgentView.tsx
--- a/Engine/Source/Programs/Horde/HordeDashboard/src/components/agents/AgentView.tsx
+++ b/Engine/Source/Programs/Horde/HordeDashboard/src/components/agents/AgentView.tsx
@@ -2056,18 +2056,21 @@ const SearchUpdate: React.FC = observer(() => {
    return null;
 });
 
-const AgentViewInner: React.FC<{ agentId?: string, poolId?: string, searchParams?: URLSearchParams, agentView?: boolean }> = observer(({ agentId, poolId, searchParams, agentView }) => {
+const AgentViewInner: React.FC<{ agentId?: string, poolId?: string, searchParams?: URLSearchParams, agentView?: boolean, agents?: AgentData[] }> = observer(({ agentId, poolId, searchParams, agentView, agents: externalAgents }) => {
 
    poolId = poolId?.toLowerCase();
 
+   const hasExternalAgents = !!externalAgents;
+
    const [initAgentUpdater, setInitAgentUpdater] = useState(false);
 
    useEffect(() => {
+      if (hasExternalAgents) return;
       const interval = setInterval(() => {
          agentStore.update(true)
       }, 15000);
       return () => clearInterval(interval);
-   }, []);
+   }, [hasExternalAgents]);
 
    // subscribe
    if (localState.searchUpdated) { }
@@ -2075,7 +2078,8 @@ const AgentViewInner: React.FC<{ agentId?: string, poolId?: string, searchParams
    const agentStyles = getAgentStyles();
 
    if (!initAgentUpdater) {
-      agentStore.update().then(() => {
+      const initPromise = hasExternalAgents ? agentStore.updatePools() : agentStore.update();
+      initPromise.then(() => {
          localState.resetState();
          if (searchParams) {
             localState.stateFromSearch(searchParams);
@@ -2092,9 +2096,11 @@ const AgentViewInner: React.FC<{ agentId?: string, poolId?: string, searchParams
       </Stack>
    }
 
+   const effectiveAgents = externalAgents ?? agentStore.agents;
+
    let activeAgent: AgentData | undefined = undefined;
    if (agentId) {
-      activeAgent = agentStore.agents.find(agent => agent.id === agentId);
+      activeAgent = effectiveAgents.find(agent => agent.id === agentId);
    }
 
    // if(activeAgent) {
@@ -2499,14 +2505,14 @@ const AgentViewInner: React.FC<{ agentId?: string, poolId?: string, searchParams
       localState.setAgentId(undefined);
    }
 
-   const totalPoolSize = agentStore.agents.filter((a) => {
+   const totalPoolSize = effectiveAgents.filter((a) => {
       if (!a.pools) {
          return false;
       }
       return !poolId || a.pools.indexOf(poolId) !== -1
    }).length;
 
-   let agentItems = agentStore.agents.filter(filterAgents, localState.agentFilter).sort(sortAgents);
+   let agentItems = effectiveAgents.filter(filterAgents, localState.agentFilter).sort(sortAgents);
 
    if (poolId) {
       agentItems = agentItems.filter((item => {
@@ -3037,7 +3043,7 @@ const AgentViewInner: React.FC<{ agentId?: string, poolId?: string, searchParams
    }
 });
 
-export const AgentPanel: React.FC<{ agentId?: string, poolId?: string, agentView?: boolean }> = observer(({ agentId, agentView, poolId }) => {
+export const AgentPanel: React.FC<{ agentId?: string, poolId?: string, agentView?: boolean, agents?: AgentData[] }> = observer(({ agentId, agentView, poolId, agents }) => {
 
    // adjust automatically to viewport changes
    useWindowSize();
@@ -3049,7 +3055,7 @@ export const AgentPanel: React.FC<{ agentId?: string, poolId?: string, agentView
    localState.updateColumns();
 
    return <Stack>
-      <AgentViewInner agentId={localState.searchState?.agentId} agentView={agentView} poolId={poolId} />
+      <AgentViewInner agentId={localState.searchState?.agentId} agentView={agentView} poolId={poolId} agents={agents} />
    </Stack>
 
 })
