[Test Automation Hub] Truncate metadata value that are too long

#rnx
#Synced-CL 52130410

diff --git a/Engine/Source/Programs/Horde/HordeDashboard/plugins/test-automation/testAutomationCommon.tsx b/Engine/Source/Programs/Horde/HordeDashboard/plugins/test-automation/testAutomationCommon.tsx
--- a/Engine/Source/Programs/Horde/HordeDashboard/plugins/test-automation/testAutomationCommon.tsx
+++ b/Engine/Source/Programs/Horde/HordeDashboard/plugins/test-automation/testAutomationCommon.tsx
@@ -736,7 +736,7 @@ export const CheckListOption: React.FC<{options: ICheckListOption[], onChange: (
                                                                 <Stack style={{ paddingRight: 5 }}>
                                                                     <FontIcon style={{ fontSize: 12, color: item.color }} iconName={!!item.color? "Square" : "MultiSelect"} />
                                                                 </Stack>
-                                                                <Text>{props?.label}</Text>
+                                                                <Text>{props?.label?.substring(0, 16)}{props?.label && props.label.length > 16 && '...'}</Text>
                                                             </Stack>
                                                 }}
                                             />)}
@@ -796,7 +796,7 @@ export const CheckListOptionSmall: React.FC<{options: ICheckListOption[], onChan
                                         }
                                         onRenderLabel={(props) => {
                                             return <Stack verticalAlign="center" tokens={{childrenGap: 2}} style={{background: item.color ? `${item.color}60` : undefined}}>
-                                                        <Text style={{padding: '0 4px'}}>{props?.label}</Text>
+                                                        <Text style={{padding: '0 4px'}}>{props?.label?.substring(0, 16)}{props?.label && props.label.length > 16 && '...'}</Text>
                                                         <Stack style={{ height: 3, backgroundColor: item.color }}></Stack>
                                                     </Stack>
                                         }}
diff --git a/Engine/Source/Programs/Horde/HordeDashboard/plugins/test-automation/testAutomationPhaseGridView.tsx b/Engine/Source/Programs/Horde/HordeDashboard/plugins/test-automation/testAutomationPhaseGridView.tsx
--- a/Engine/Source/Programs/Horde/HordeDashboard/plugins/test-automation/testAutomationPhaseGridView.tsx
+++ b/Engine/Source/Programs/Horde/HordeDashboard/plugins/test-automation/testAutomationPhaseGridView.tsx
@@ -274,7 +274,7 @@ const MetaSelectorTitle: React.FC<{ stateContext: MetaSelectorContext, onFilterC
                                         <Stack style={{background: stateContext.metaValueColors?.get(v, "60"), cursor: "pointer"}}
                                             onClick={() => onMetaValueBeforeMainClick(v)}
                                         >
-                                            <Text style={{whiteSpace: 'nowrap', padding: '0 10px 0 4px', position: "relative", height: 18}}>{v}
+                                            <Text style={{whiteSpace: 'nowrap', padding: '0 10px 0 4px', position: "relative", height: 18}}>{v.substring(0, 16)}{v.length > 16 && '...'}
                                                 {!!options.find(o => o.text === v)?.checked &&
                                                     <FontIcon iconName="Filter" style={{fontSize: 8, position: "absolute", right: 2, top: 2, color: styles.highlightFilter.backgroundColor}} />
                                                 }
@@ -297,7 +297,7 @@ const MetaSelectorTitle: React.FC<{ stateContext: MetaSelectorContext, onFilterC
                                         <Stack style={{background: stateContext.metaValueColors?.get(v, "60"), cursor: "pointer", position: "relative"}}
                                             onClick={() => onMetaValueClick(v)}
                                         >
-                                            <Text style={{whiteSpace: 'nowrap', padding: '0 10px 0 4px', position: "relative", height: 18}}>{v}
+                                            <Text style={{whiteSpace: 'nowrap', padding: '0 10px 0 4px', position: "relative", height: 18}}>{v.substring(0, 16)}{v.length > 16 && '...'}
                                                 {!!options.find(o => o.text === v)?.checked &&
                                                     <FontIcon iconName="Filter" style={{fontSize: 8, position: "absolute", right: 2, top: 2, color: styles.highlightFilter.backgroundColor}} />
                                                 }
@@ -333,7 +333,7 @@ const MetaSelectorTitle: React.FC<{ stateContext: MetaSelectorContext, onFilterC
                                                         transformOrigin: '0% 100%',
                                                         position: 'absolute',
                                                         left: isTooTight ? 8 : 4,
-                                                    }}>{v ?? " "}</Text>
+                                                    }}>{v?.substring(0, 16) ?? " "}{v?.length > 16 && '...'}</Text>
                                                 </Stack>
                                                 <Stack style={{background: v ? stateContext.metaValueColors?.get(v) : undefined, height: 3}}></Stack>
                                             </Stack>
