Horde: Fix step status colors in label hover card for skipped steps
The StepsSection in LabelHoverCard had inline color logic that did not
handle Skipped or Aborted step states, causing them to show as failed
or succeeded. Use the existing getStepStatusColor() which handles all
step states correctly.

diff --git a/Engine/Source/Programs/Horde/HordeDashboard/src/components/labels/LabelHoverCard.tsx b/Engine/Source/Programs/Horde/HordeDashboard/src/components/labels/LabelHoverCard.tsx
--- a/Engine/Source/Programs/Horde/HordeDashboard/src/components/labels/LabelHoverCard.tsx
+++ b/Engine/Source/Programs/Horde/HordeDashboard/src/components/labels/LabelHoverCard.tsx
@@ -6,7 +6,7 @@ import { Link } from 'react-router-dom';
 import backend from '../../backend';
 import { GetLabelDetailsResponse, GetLabelIssueResponse, GetLabelStepResponse, GetLogEventResponse, IssueSeverity, JobStepOutcome, JobStepState, LabelOutcome, LabelState } from '../../backend/Api';
 import dashboard, { StatusColor } from '../../backend/Dashboard';
-import { getLabelColor } from '../../styles/colors';
+import { getLabelColor, getStepStatusColor } from '../../styles/colors';
 import { ChangeButton } from '../ChangeButton';
 import { IssueModalV2 } from '../IssueViewV2';
 
@@ -266,7 +266,6 @@ const StepsSection: React.FC<{
     expandedLogs: Set<string>;
     setExpandedLogs: (s: Set<string>) => void;
 }> = ({ steps, logEvents, totalStepCount, jobId, expandedLogs, setExpandedLogs }) => {
-    const colors = dashboard.getStatusColors();
     const maxDisplay = 8;
     const failingSteps = steps.filter(s => s.outcome === JobStepOutcome.Failure || s.outcome === JobStepOutcome.Warnings);
     const displaySteps = failingSteps.length > 0 ? failingSteps : steps;
@@ -277,13 +276,7 @@ const StepsSection: React.FC<{
                 {failingSteps.length > 0 ? 'Failing Steps' : 'Steps'}
             </Text>
             {displaySteps.slice(0, maxDisplay).map(step => {
-                const stepColor = step.outcome === JobStepOutcome.Failure
-                    ? colors.get(StatusColor.Failure)
-                    : step.outcome === JobStepOutcome.Warnings
-                        ? colors.get(StatusColor.Warnings)
-                        : (step.state === JobStepState.Running || step.state === JobStepState.Ready || step.state === JobStepState.Waiting)
-                            ? colors.get(StatusColor.Running)
-                            : colors.get(StatusColor.Success);
+                const stepColor = getStepStatusColor(step.state, step.outcome);
                 const logEntry = step.logId ? logEvents[step.logId] : undefined;
                 const isExpanded = step.logId ? expandedLogs.has(step.logId) : false;
 
