Add tick icon when Test Health state is back to healthy

#rnx
#Synced-CL 52169528

diff --git a/Engine/Source/Programs/Horde/Plugins/Build/HordeServer.Build/Notifications/Sinks/SlackNotificationSink.cs b/Engine/Source/Programs/Horde/Plugins/Build/HordeServer.Build/Notifications/Sinks/SlackNotificationSink.cs
--- a/Engine/Source/Programs/Horde/Plugins/Build/HordeServer.Build/Notifications/Sinks/SlackNotificationSink.cs
+++ b/Engine/Source/Programs/Horde/Plugins/Build/HordeServer.Build/Notifications/Sinks/SlackNotificationSink.cs
@@ -2935,6 +2935,7 @@ namespace HordeServer.Notifications.Sinks
 					SlackMessage fixedMessage = new SlackMessage();
 					fixedMessage.Text = $"Test Health was restored to {report.State}.";
 					await _slackClient.PostMessageToThreadAsync(threadId, fixedMessage, cancellationToken);
+					await _slackClient.AddReactionAsync(threadId, "tick", cancellationToken);
 
 					// delete message state
 					await DeleteMessageStateAsync(recipient, eventId, cancellationToken);
