Horde: Fix WriteAliases 403 Forbidden after .NET 10 upgrade

diff --git a/Engine/Source/Programs/Horde/HordeServer/HordeServer.csproj b/Engine/Source/Programs/Horde/HordeServer/HordeServer.csproj
--- a/Engine/Source/Programs/Horde/HordeServer/HordeServer.csproj
+++ b/Engine/Source/Programs/Horde/HordeServer/HordeServer.csproj
@@ -30,13 +30,13 @@
       <PrivateAssets>all</PrivateAssets>
       <IncludeAssets>runtime; build; native; contentfiles; analyzers; buildtransitive</IncludeAssets>
     </PackageReference>
-    <PackageReference Include="Microsoft.AspNetCore.Authentication.JwtBearer" Version="8.0.20" />
-    <PackageReference Include="Microsoft.AspNetCore.Authentication.OpenIdConnect" Version="8.0.20" />
-    <PackageReference Include="Microsoft.AspNetCore.DataProtection.StackExchangeRedis" Version="8.0.20" />
-    <PackageReference Include="Microsoft.AspNetCore.Mvc.Razor.RuntimeCompilation" Version="8.0.20" />
-    <PackageReference Include="Microsoft.AspNetCore.SpaServices.Extensions" Version="8.0.20" />
-    <PackageReference Include="Microsoft.Extensions.Hosting.WindowsServices" Version="8.0.1" />
-    <PackageReference Include="Microsoft.Extensions.Http.Polly" Version="8.0.20" />
+    <PackageReference Include="Microsoft.AspNetCore.Authentication.JwtBearer" Version="10.0.3" />
+    <PackageReference Include="Microsoft.AspNetCore.Authentication.OpenIdConnect" Version="10.0.3" />
+    <PackageReference Include="Microsoft.AspNetCore.DataProtection.StackExchangeRedis" Version="10.0.3" />
+    <PackageReference Include="Microsoft.AspNetCore.Mvc.Razor.RuntimeCompilation" Version="10.0.3" />
+    <PackageReference Include="Microsoft.AspNetCore.SpaServices.Extensions" Version="10.0.3" />
+    <PackageReference Include="Microsoft.Extensions.Hosting.WindowsServices" Version="10.0.3" />
+    <PackageReference Include="Microsoft.Extensions.Http.Polly" Version="10.0.3" />
     <PackageReference Include="Microsoft.VisualStudio.Threading.Analyzers" Version="17.14.15">
       <PrivateAssets>all</PrivateAssets>
       <IncludeAssets>runtime; build; native; contentfiles; analyzers; buildtransitive</IncludeAssets>
diff --git a/Engine/Source/Programs/Horde/Plugins/Storage/HordeServer.Storage/Storage/StorageController.cs b/Engine/Source/Programs/Horde/Plugins/Storage/HordeServer.Storage/Storage/StorageController.cs
--- a/Engine/Source/Programs/Horde/Plugins/Storage/HordeServer.Storage/Storage/StorageController.cs
+++ b/Engine/Source/Programs/Horde/Plugins/Storage/HordeServer.Storage/Storage/StorageController.cs
@@ -36,11 +36,12 @@ namespace HordeServer.Storage
 		readonly IAclService _aclService;
 		readonly GcMetricsCollection _gcMetrics;
 		readonly CollectionStatsService _collectionStats;
+		readonly ILogger _logger;
 
 		/// <summary>
 		/// Constructor
 		/// </summary>
-		public StorageController(StorageService storageService, IOptionsSnapshot<StorageConfig> storageConfig, IOptions<StorageServerConfig> staticStorageConfig, IAclService aclService, GcMetricsCollection gcMetrics, CollectionStatsService collectionStats)
+		public StorageController(StorageService storageService, IOptionsSnapshot<StorageConfig> storageConfig, IOptions<StorageServerConfig> staticStorageConfig, IAclService aclService, GcMetricsCollection gcMetrics, CollectionStatsService collectionStats, ILogger<StorageController> logger)
 		{
 			_storageService = storageService;
 			_storageConfig = storageConfig;
@@ -48,6 +49,7 @@ namespace HordeServer.Storage
 			_aclService = aclService;
 			_gcMetrics = gcMetrics;
 			_collectionStats = collectionStats;
+			_logger = logger;
 		}
 
 		bool Authorize(NamespaceId namespaceId, AclAction action)
@@ -57,7 +59,8 @@ namespace HordeServer.Storage
 
 		bool AuthorizePaths(NamespaceId namespaceId, AclAction action, IEnumerable<string> paths)
 		{
-			if (!paths.Any())
+			List<string> pathList = paths.ToList();
+			if (pathList.Count == 0)
 			{
 				return true;
 			}
@@ -65,10 +68,18 @@ namespace HordeServer.Storage
 			{
 				return true;
 			}
-			if (paths.All(x => HasPathClaim(User, HordeClaimTypes.WriteNamespace, namespaceId, x)))
+			if (pathList.All(x => HasPathClaim(User, HordeClaimTypes.WriteNamespace, namespaceId, x)))
 			{
 				return true;
 			}
+
+			_logger.LogWarning(
+				"AuthorizePaths denied {Action} for namespace {NamespaceId}. Paths=[{Paths}], Claims=[{Claims}]",
+				action,
+				namespaceId,
+				String.Join(", ", pathList),
+				String.Join("; ", User.Claims.Select(c => $"{c.Type}={c.Value}"))
+			);
 			return false;
 		}
 
