Horde: Fix JWT validation after .NET 10 upgrade - UseSecurityTokenValidators

JwtAuthHandler registers via AddScheme which bypasses PostConfigureOptions. This means
TokenHandlers is never populated with a JsonWebTokenHandler. On .NET 8, JwtBearerHandler
had a fallback: empty TokenHandlers would fall back to SecurityTokenValidators (which has
a default JwtSecurityTokenHandler from the JwtBearerOptions constructor). .NET 10 removed
that fallback, so empty TokenHandlers + UseSecurityTokenValidators=false means no token
validation occurs - claims are lost and all storage path authorization fails with 403.

Set UseSecurityTokenValidators=true to explicitly use the JwtSecurityTokenHandler path,
which also matches how tokens are created in AclService.IssueBearerTokenAsync.

diff --git a/Engine/Source/Programs/Horde/HordeServer/Authentication/JwtAuthHandler.cs b/Engine/Source/Programs/Horde/HordeServer/Authentication/JwtAuthHandler.cs
--- a/Engine/Source/Programs/Horde/HordeServer/Authentication/JwtAuthHandler.cs
+++ b/Engine/Source/Programs/Horde/HordeServer/Authentication/JwtAuthHandler.cs
@@ -41,6 +41,13 @@ namespace HordeServer.Authentication
 			ConfigureNamedOptions<JwtBearerOptions> namedOptions = new (AuthenticationScheme, options =>
 			{
 				options.RequireHttpsMetadata = !settings.OidcDebugMode;
+
+				// This handler bypasses PostConfigureOptions (line below passes empty array), so
+				// TokenHandlers is never populated with a JsonWebTokenHandler. On .NET 8, empty
+				// TokenHandlers fell back to SecurityTokenValidators (which has a default
+				// JwtSecurityTokenHandler). .NET 10 removed that fallback, so without this flag
+				// no token validation occurs and claims are lost.
+				options.UseSecurityTokenValidators = true;
 			});
 			OptionsFactory<JwtBearerOptions> optionsFactory = new (new[] { namedOptions }, Array.Empty<IPostConfigureOptions<JwtBearerOptions>>());
 			return new OptionsMonitor<JwtBearerOptions>(optionsFactory, Array.Empty<IOptionsChangeTokenSource<JwtBearerOptions>>(), optionsCache);
