Horde: Normalize IPv6-mapped IPv4 addresses to plain IPv4 for compute IP comparison

#Synced-CL 0

diff --git a/Engine/Source/Programs/Horde/Plugins/Compute/HordeServer.Compute.Tests/Compute/ComputeServiceTests.cs b/Engine/Source/Programs/Horde/Plugins/Compute/HordeServer.Compute.Tests/Compute/ComputeServiceTests.cs
new file mode 100644
--- /dev/null
+++ b/Engine/Source/Programs/Horde/Plugins/Compute/HordeServer.Compute.Tests/Compute/ComputeServiceTests.cs
@@ -0,0 +1,19 @@
+// Copyright Epic Games, Inc. All Rights Reserved.
+
+using System.Net;
+
+namespace HordeServer.Tests.Compute
+{
+	[TestClass]
+	public class ComputeServiceTests
+	{
+		[TestMethod]
+		[DataRow("192.168.1.1")]
+		[DataRow("::ffff:192.168.1.1")]
+		public void TryGetNetworkConfigMatchesCidrBlock(string ip)
+		{
+			ComputeConfig config = new() { Networks = { new NetworkConfig { Id = "test", CidrBlock = "192.168.0.0/16" } } };
+			Assert.IsTrue(config.TryGetNetworkConfig(IPAddress.Parse(ip), out _));
+		}
+	}
+}
diff --git a/Engine/Source/Programs/Horde/Plugins/Compute/HordeServer.Compute/Compute/ComputeService.cs b/Engine/Source/Programs/Horde/Plugins/Compute/HordeServer.Compute/Compute/ComputeService.cs
--- a/Engine/Source/Programs/Horde/Plugins/Compute/HordeServer.Compute/Compute/ComputeService.cs
+++ b/Engine/Source/Programs/Horde/Plugins/Compute/HordeServer.Compute/Compute/ComputeService.cs
@@ -420,7 +420,12 @@ namespace HordeServer.Compute
 				// Do not assign a compute resource hosted on the same machine as requester
 				// This is often the case when a user is offering their local workstation to participate in a shared compute pool
 				// With UnrealBuildAccelerator, it already makes full use of the local machine.
-				conditions.Add(Condition.Parse($"{KnownPropertyNames.ComputeIp} != '{arp.RequesterIp.ToString()}'"));
+				// Normalize IPv6-mapped IPv4 addresses (e.g. ::ffff:192.168.1.1) to plain IPv4
+				// to ensure consistent string comparison with agent ComputeIp values
+				IPAddress requesterIp = arp.RequesterIp.IsIPv4MappedToIPv6
+					? arp.RequesterIp.MapToIPv4()
+					: arp.RequesterIp;
+				conditions.Add(Condition.Parse($"{KnownPropertyNames.ComputeIp} != '{requesterIp}'"));
 			}
 			
 			Task<UbaComputeResource>? ubaCache = null;
