[Test Automation Hub] Optimize closest Data by grouping the sessions into smaller time section over the timeline
+ Filter data instead of changing each item visibility

#rnx
#Synced-CL 52253672

diff --git a/Engine/Source/Programs/Horde/HordeDashboard/plugins/test-automation/phaseHistoryGraph.ts b/Engine/Source/Programs/Horde/HordeDashboard/plugins/test-automation/phaseHistoryGraph.ts
--- a/Engine/Source/Programs/Horde/HordeDashboard/plugins/test-automation/phaseHistoryGraph.ts
+++ b/Engine/Source/Programs/Horde/HordeDashboard/plugins/test-automation/phaseHistoryGraph.ts
@@ -7,13 +7,14 @@ import { getShortNiceTime } from 'horde/base/utilities/timeUtils';
 import { getHordeStyling } from 'horde/styles/Styles';
 import { getPhaseSessionStatusColor,  SessionListSelector } from "./testAutomationCommon";
 import { renderToString } from "react-dom/server";
-import { timelineHandler, Margin, constructTimeline, calculateAverageSessionFrequencyHours } from "./timelineGraph";
+import { timelineHandler, Margin, constructTimeline, calculateAverageSessionFrequencyHours, getSessionsNearTimeFactory } from "./timelineGraph";
 
 // Handle bad "@types/d3" types, fix if addressed upstream
 const _d3 = d3 as any;
 
 type SelectionType = d3.Selection<SVGGElement, unknown, null, undefined>;
 type DivSelectionType = d3.Selection<HTMLDivElement, unknown, HTMLElement, undefined>;
+type SelectionBaseType = d3.Selection<d3.BaseType, unknown, null, undefined>;
 type FingerprintSeries = { 
     fingerprint: string,
     color: string,
@@ -70,8 +71,8 @@ export class PhaseHistoryGraph {
         const handler = this.handler;
         const sessions = this.sessions;
 
-        const X = _d3.map(sessions, (r) => handler.commitIdDates.get(r.commitId)!.getTime() / 1000);
-        const Y = _d3.map(sessions, (r) => 0);
+        const X = _d3.map(sessions, (r) => handler.commitIdDates.get(r.commitId)!.getTime() / 1000) as number[];
+        const Y = _d3.map(sessions, (r) => 0) as number[];
 
         const yDomain = new _d3.InternSet(Y as any);
 
@@ -83,6 +84,8 @@ export class PhaseHistoryGraph {
         const yRange = [this.margin.top, height - this.margin.bottom];
         const yScale = _d3.scalePoint(yDomain, yRange).round(true).padding(yPadding);
 
+        const getSessionsNearTime = getSessionsNearTimeFactory(sessions, X, this.handler.filterState!.weeks! * 7);
+
         const getFingerprintSeries = () => {
             // group sessions into consecutive runs of the same error fingerprint
             const fingerprintSeries: FingerprintSeries[] = [];
@@ -222,28 +225,31 @@ export class PhaseHistoryGraph {
                  );
 
             group.selectAll(".session-circle")
-                .data(([, I]: any) => _d3.group(I, i => X[i]), ([, items]: any) => sessions[items[0]].id)
+                .data(
+                    ([, I]: any) => _d3.groups(I, i => X[i])
+                        .filter(([, items]: any) => X[items[0] as any] >= timelineHandler.startDate! && X[items[0] as any] <= timelineHandler.endDate!),
+                    ([, items]: any) => sessions[items[0]].id
+                )
                 .join(
                     enter => enter.append("g")
                         .attr("class", "session-circle")
                         .attr("id", ([, items]: any) => `circle${sessions[items[0]].id}`)
                         .attr("transform", ([, items]: any) => `translate(${timelineHandler.timelineXScale!(X[items[0] as any])},0)`)
-                        .attr("visibility", ([, items]: any) => X[items[0] as any] >= timelineHandler.startDate! && X[items[0] as any] <= timelineHandler.endDate! ? "visible" : "hidden")
                         .append("g")
                         .attr("class", "circleScale")
-                        .attr("id", ([, items]: any) => `circleScale${sessions[items[0]].id}`)
                         .style("cursor", "pointer")
                         .each(([, items]: any, i, nodes: any) => SessionCircle(nodes[i], items.map(i => sessions[i]), this.phaseSessionId, circleSize, !!this.phaseSessionId && items.some(i => sessions[i].id === this.phaseSessionId))),
-                    update => update.attr("transform", ([, items]: any) => `translate(${timelineHandler.timelineXScale!(X[items[0] as any])},0)`)
-                        .attr("visibility", ([, items]: any) => X[items[0] as any] >= timelineHandler.startDate! && X[items[0] as any] <= timelineHandler.endDate! ? "visible" : "hidden"),
+                    update => update.attr("transform", ([, items]: any) => `translate(${timelineHandler.timelineXScale!(X[items[0] as any])},0)`),
                  );
         }
         metaGroups.call(drawMetaGroups);
 
-        const reDrawCircles = (items: PhaseSessionResult[], index: number) => {
-            const circle = svg.select(`#circleScale${items[0].id}`);
-            circle.selectAll("*").remove();
-            SessionCircle(circle.node() as any, items, items[index].id, circleSize, !!this.phaseSessionId && items.some(i => i.id === this.phaseSessionId));
+        const setCircleStatus = (selection: SelectionBaseType, items: PhaseSessionResult[], index: number) => {
+            const sessionId = items[index].id;
+            const session = sessions.length === 1 || !sessionId? sessions[0] : sessions.find(s => s.id === sessionId) ?? sessions[0];
+            const color = getPhaseSessionStatusColor(session);
+            selection.select(".status")
+                .attr("fill", color as any);
         }
 
         // Attempt to select an existing div with id 'tooltip' from the body, or create one if it doesn't exist
@@ -272,7 +278,7 @@ export class PhaseHistoryGraph {
 
             y -= 16;
 
-            const closest = sessions.reduce((best, session, i) => {
+            const closest = getSessionsNearTime(timelineHandler.timelineXScale!.invert(x)).reduce((best, session: PhaseSessionResult) => {
                 const absy = Math.abs(yScale(Y[0]) - y);
                 const sessionTime = handler.commitIdDates.get(session.commitId)!.getTime() / 1000;
                 const absx = Math.abs(timelineHandler.timelineXScale!(sessionTime) - x);
@@ -280,15 +286,15 @@ export class PhaseHistoryGraph {
                 const lengthSqr = absy * absy + absx * absx;
 
                 if (lengthSqr < best.value) {
-                    return { indexes: [i], value: lengthSqr };
+                    return { sessions: [session], value: lengthSqr };
                 } else {
-                     if (lengthSqr === best.value) best.indexes.push(i);
+                     if (lengthSqr === best.value) best.sessions.push(session);
                     return best;
                 }
-            }, { indexes: [] as number[], value: Number.MAX_SAFE_INTEGER });
+            }, { sessions: [] as PhaseSessionResult[], value: Number.MAX_SAFE_INTEGER });
 
-            if (closest.indexes.length) {
-                return closest.indexes.map(i => sessions[i]);
+            if (closest.sessions.length) {
+                return closest.sessions;
             }
 
             return undefined;
@@ -331,10 +337,10 @@ export class PhaseHistoryGraph {
                 const metaIndex = 0;
                 let index = closestItems.length > 1? closestItems.findIndex(item => item.id === this.phaseSessionId) : 0;
                 if (index < 0) index = 0;
-                this.tooltipItem = { items: closestItems, index: index, sessionId: closest.id, delta: 0 };
-
-                svg.select(`#circleScale${closest.id}`).attr("transform", `scale(2)`);
-                svg.select(`#circle${closest.id}`).raise();
+                const circleSelection = svg.select(`#circle${closest.id}`);
+                circleSelection.raise();
+                circleSelection.select(".circleScale").attr("transform", `scale(2)`);
+                this.tooltipItem = { items: closestItems, index: index, sessionId: closest.id, delta: 0, circleSelection: circleSelection };
 
                 const timeStamp = handler.commitIdDates.get(closest.commitId)!.getTime() / 1000;
                 const tx = timelineHandler.timelineXScale!(timeStamp);
@@ -354,11 +360,14 @@ export class PhaseHistoryGraph {
         const handleMouseLeave = (event: any) => {
 
             if (this.tooltipItem) {
-                svg.selectAll(".circleScale").attr("transform", `scale(1)`);
                 tooltip.style("display", "none");
-                let index = this.tooltipItem.items.length > 1? this.tooltipItem.items.findIndex(item => item.id === this.phaseSessionId) : 0;
-                if (index < 0) index = 0;
-                reDrawCircles(this.tooltipItem.items, index);
+                const circleSelection = this.tooltipItem.circleSelection;
+                circleSelection.select(".circleScale").attr("transform", `scale(1)`);
+                if (this.tooltipItem.items.length > 1) {
+                    let index = this.tooltipItem.items.findIndex(item => item.id === this.phaseSessionId);
+                    if (index < 0) index = 0;
+                    setCircleStatus(circleSelection, this.tooltipItem.items, index);
+                }
                 this.tooltipItem = undefined;
             }
 
@@ -393,7 +402,7 @@ export class PhaseHistoryGraph {
                         this.tooltipItem.index -= this.tooltipItem.items.length;
                     }
 
-                    reDrawCircles(this.tooltipItem.items, this.tooltipItem.index);
+                    setCircleStatus(this.tooltipItem.circleSelection, this.tooltipItem.items, this.tooltipItem.index);
                     const desc = getInfoSession(this.tooltipItem.items, this.tooltipItem.index);
                     this.updateTooltip(undefined, undefined, desc);
                 }
@@ -469,7 +478,7 @@ export class PhaseHistoryGraph {
     sessions: PhaseSessionResult[] = [];
 
     tooltip?: DivSelectionType;
-    tooltipItem?: {index: number, items: PhaseSessionResult[], sessionId: string, delta: number};
+    tooltipItem?: {index: number, items: PhaseSessionResult[], sessionId: string, delta: number, circleSelection: SelectionBaseType};
 
     private _averageSessionFrequencyHours?: number;
 }
@@ -479,6 +488,7 @@ const SessionCircle = (container: SVGGElement, sessions: PhaseSessionResult[], s
     const color = getPhaseSessionStatusColor(session);
     const g = d3.select(container);
     g.append("circle")
+        .attr("class", "status")
         .attr("stroke", "none")
         .attr("fill", color as any)
         .attr("r", radius);
diff --git a/Engine/Source/Programs/Horde/HordeDashboard/plugins/test-automation/testHistoryGraph.ts b/Engine/Source/Programs/Horde/HordeDashboard/plugins/test-automation/testHistoryGraph.ts
--- a/Engine/Source/Programs/Horde/HordeDashboard/plugins/test-automation/testHistoryGraph.ts
+++ b/Engine/Source/Programs/Horde/HordeDashboard/plugins/test-automation/testHistoryGraph.ts
@@ -3,18 +3,19 @@
 import * as d3 from "d3";
 import { TestDataHandler, TestNameRef, MetadataRef, TestMetaStatus, TestSessionResult } from './testData';
 import dashboard, { StatusColor } from 'horde/backend/Dashboard';
-import { getHumanTime, getShortNiceTime, msecToElapsed } from 'horde/base/utilities/timeUtils';
+import { getShortNiceTime, msecToElapsed } from 'horde/base/utilities/timeUtils';
 import { getHordeStyling } from 'horde/styles/Styles';
 import { getStatusColors, statusTexts, SessionValues, SessionListSelector } from "./testAutomationCommon";
 import { renderToString } from "react-dom/server"
 import { TestOutcome } from "./api";
-import { timelineHandler, Margin, constructTimeline, calculateAverageSessionFrequencyHours } from "./timelineGraph";
+import { timelineHandler, Margin, constructTimeline, calculateAverageSessionFrequencyHours, getSessionsNearTimeFactory } from "./timelineGraph";
 
 // Handle bad "@types/d3" types, fix if addressed upstream
 const _d3 = d3 as any;
 
 type SelectionType = d3.Selection<SVGGElement, unknown, null, undefined>;
 type DivSelectionType = d3.Selection<HTMLDivElement, unknown, HTMLElement, undefined>;
+type SelectionBaseType = d3.Selection<d3.BaseType, unknown, null, undefined>;
 
 const clampMin = (x: number) => x !== 0 && x <= 0.3? 0.3: x;
 
@@ -100,6 +101,8 @@ export class TestHistoryGraph {
 
         const yScale = _d3.scalePoint(yDomain, yRange).round(true).padding(yPadding);
 
+        const getSessionsNearTime = getSessionsNearTimeFactory(sessions, X, this.handler.filterState!.weeks! * 7, this.metaInfo.size);
+
         const svg = d3.select(container)
             .append("svg")
             .attr("width", width)
@@ -143,32 +146,32 @@ export class TestHistoryGraph {
                                 .style("font-size", 10)
                                 .attr("dy", "0.15em") // center stream name
                                 .attr("fill", dashboard.darktheme ? "#E0E0E0" : "#2D3F5F")
-                                .text(([y]: any) => this.metaNames.get(this.metaRefs[y])!);
+                                .text(([y]: any) => this.metaNames.get(this.metaRefs[y])!.substring(0, 32));
                         }),
                     update => update.attr("transform", ([y]: any) => `translate(0,${(yScale(y) as any) + 16})`),
                  );
 
             group.selectAll(".session-circle")
-                .data(([, I]: any) => _d3.group(I, i => X[i]), ([, items]: any) => sessions[items[0]].id)
+                .data(
+                    ([, I]: any) => _d3.groups(I, i => X[i])
+                        .filter(([, items]: any) => X[items[0] as any] >= timelineHandler.startDate! && X[items[0] as any] <= timelineHandler.endDate!),
+                    ([, items]: any) => sessions[items[0]].id)
                 .join(
                     enter => enter.append("g")
                         .attr("class", "session-circle")
                         .attr("id", ([, items]: any) => `circle${sessions[items[0]].id}`)
                         .attr("transform", ([, items]: any) => `translate(${timelineHandler.timelineXScale!(X[items[0] as any])},0)`)
-                        .attr("visibility", ([, items]: any) => X[items[0] as any] >= timelineHandler.startDate! && X[items[0] as any] <= timelineHandler.endDate! ? "visible" : "hidden")
                         .append("g")
                         .attr("class", "circleScale")
-                        .attr("id", ([, items]: any) => `circleScale${sessions[items[0]].id}`)
                         .style("cursor", "pointer")
                         .each(([, items]: any, i, nodes: any) => SessionConcentricCircle(nodes[i], items.map(i => sessions[i]), 0, circleSize)),
-                    update => update.attr("transform", ([, items]: any) => `translate(${timelineHandler.timelineXScale!(X[items[0] as any])},0)`)
-                        .attr("visibility", ([, items]: any) => X[items[0] as any] >= timelineHandler.startDate! && X[items[0] as any] <= timelineHandler.endDate! ? "visible" : "hidden"),
+                    update => update.attr("transform", ([, items]: any) => `translate(${timelineHandler.timelineXScale!(X[items[0] as any])},0)`),
                  );
         }
         metaGroups.call(drawMetaGroups);
 
-        const reDrawCircles = (items: TestSessionResult[], index: number) => {
-            const circle = svg.select(`#circleScale${items[0].id}`);
+        const setCircleStatus = (selection: SelectionBaseType, items: TestSessionResult[], index: number) => {
+            const circle = selection.select(".circleScale");
             circle.selectAll("*").remove();
             SessionConcentricCircle(circle.node() as any, items, index, circleSize);
         }
@@ -202,23 +205,25 @@ export class TestHistoryGraph {
 
             y -= 16;
 
-            const closest = sessions.reduce((best, session, i) => {
+            const closest = getSessionsNearTime(timelineHandler.timelineXScale!.invert(x)).reduce((best, session: TestSessionResult) => {
                 const absy = Math.abs(this.metaInfo.get(session.metadataId)!.y - y);
+                if (absy > 9) return best; // skip if too far away from the mouse y position
+
                 const sessionTime = handler.commitIdDates.get(session.commitId)!.getTime() / 1000;
                 const absx = Math.abs(timelineHandler.timelineXScale!(sessionTime) - x);
 
                 const lengthSqr = absy * absy + absx * absx;
 
                 if (lengthSqr < best.value) {
-                    return { indexes: [i], value: lengthSqr, meta: session.metadataId };
+                    return { sessions: [session], value: lengthSqr, meta: session.metadataId };
                 } else {
-                    if (lengthSqr === best.value && session.metadataId === best.meta) best.indexes.push(i);
+                    if (lengthSqr === best.value && session.metadataId === best.meta) best.sessions.push(session);
                     return best;
                 }
-            }, { indexes: [] as number[], value: Number.MAX_SAFE_INTEGER, meta: '' });
+            }, { sessions: [] as TestSessionResult[], value: Number.MAX_SAFE_INTEGER, meta: '' });
 
-            if (closest.indexes.length) {
-                return closest.indexes.map(i => sessions[i]);
+            if (closest.sessions.length) {
+                return closest.sessions;
             }
 
             return undefined;
@@ -228,7 +233,7 @@ export class TestHistoryGraph {
         const getInfoSession = (items: TestSessionResult[], index: number) => {
             const item = items[index];
             const metaIndex = this.metaInfo.get(item.metadataId)!.index;
-            const cmeta = this.metaRefs[metaIndex].getValuesExcept(this.commonMetaKeys!).join(" / ");
+            const cmeta = this.metaRefs[metaIndex].getValuesExcept(this.commonMetaKeys!).join(" / ").substring(0, 64);
             const date = getShortNiceTime(item.start, true, true);
             let desc = "";
             if (items.length > 1) {
@@ -260,11 +265,13 @@ export class TestHistoryGraph {
             if (this.tooltipItem?.sessionId !== closest.id) {
                 handleMouseLeave(undefined);
 
-                this.tooltipItem = {sessionId: closest.id, index: 0, items: closestItems, delta: 0};
+                const circleSelection = svg.select(`#circle${closest.id}`);
+                circleSelection.raise();
+                circleSelection.select(".circleScale").attr("transform", `scale(2.5)`);
+
+                this.tooltipItem = {sessionId: closest.id, index: 0, items: closestItems, delta: 0, circleSelection: circleSelection};
                 const metaIndex = this.metaInfo.get(closest.metadataId)!.index;
 
-                svg.select(`#circleScale${closest.id}`).attr("transform", `scale(2.5)`);
-                svg.select(`#circle${closest.id}`).raise();
                 svg.select(`#meta${metaIndex}`).raise();
 
                 const timeStamp = handler.commitIdDates.get(closest.commitId)!.getTime() / 1000;
@@ -286,9 +293,12 @@ export class TestHistoryGraph {
         const handleMouseLeave = (event: any) => {
 
             if (this.tooltipItem) {
-                svg.selectAll(".circleScale").attr("transform", `scale(1)`);
                 tooltip.style("display", "none");
-                reDrawCircles(this.tooltipItem.items, 0);
+                const circleSelection = this.tooltipItem.circleSelection;
+                circleSelection.select(".circleScale").attr("transform", `scale(1)`);
+                if (this.tooltipItem.items.length > 1) {
+                    setCircleStatus(circleSelection, this.tooltipItem.items, 0);
+                }
                 this.tooltipItem = undefined;
             }
 
@@ -323,7 +333,7 @@ export class TestHistoryGraph {
                         this.tooltipItem.index -= this.tooltipItem.items.length;
                     }
 
-                    reDrawCircles(this.tooltipItem.items, this.tooltipItem.index);
+                    setCircleStatus(this.tooltipItem.circleSelection, this.tooltipItem.items, this.tooltipItem.index);
                     const desc = getInfoSession(this.tooltipItem.items, this.tooltipItem.index);
                     this.updateTooltip(undefined, undefined, desc);
                 }
@@ -404,7 +414,7 @@ export class TestHistoryGraph {
     clipId: string;
 
     tooltip?: DivSelectionType;
-    tooltipItem?: {index: number, items: TestSessionResult[], sessionId: string, delta: number};
+    tooltipItem?: {index: number, items: TestSessionResult[], sessionId: string, delta: number, circleSelection: SelectionBaseType};
 
     private _averageSessionFrequencyHours?: number;
 }
@@ -466,7 +476,7 @@ const ConcentricCircle = (container: SVGGElement, data: ItemStack[], radius: num
         .range(data.map(d => d.color));
     
     const g = d3.select(container)
-        .selectAll("centric")
+        .selectAll(".centric")
         .data(data)
         .join("circle")
             .attr("class", "centric")
diff --git a/Engine/Source/Programs/Horde/HordeDashboard/plugins/test-automation/timelineGraph.ts b/Engine/Source/Programs/Horde/HordeDashboard/plugins/test-automation/timelineGraph.ts
--- a/Engine/Source/Programs/Horde/HordeDashboard/plugins/test-automation/timelineGraph.ts
+++ b/Engine/Source/Programs/Horde/HordeDashboard/plugins/test-automation/timelineGraph.ts
@@ -24,9 +24,8 @@ class TimelineHandler {
         {
             this._testid = testid;
             if (!this.minDate || !this.maxDate) {
-                const dates = handler.commitIdDates.values().map(d => d.getTime() / 1000).toArray();
-                this.minDate = Math.min(...dates);
-                this.maxDate = Math.max(...dates);
+                this.maxDate = Date.now() / 1000;
+                this.minDate = this.maxDate - 60 * 60 * 24 * 7 * handler.filterState.weeks!;
                 this.ticks = this.getTicks();
                 this.axisXScale = _d3.scaleTime().domain([this.minDate, this.maxDate]);
             }
@@ -120,12 +119,16 @@ class TimelineHandler {
         return this.endDate === this.maxDate && this.startDate === this.minDate;
     }
 
-    rescaleTimeline(factor: number) {
+    get scaleFactor() {
+        return this.getTimelineRangeWidth() / this.axisRangeWidth!;
+    }
+
+    rescaleTimeline(factor: number, endDate?: number) {
         if (!this.isTimelineAtExtreme) {
             // if timeline is not at the extremes, don't rescale
             return;
         }
-        this.startDate = this.endDate! - ((this.endDate! - this.startDate!) * factor);
+        this.startDate = (endDate ?? this.endDate!) - ((this.endDate! - this.startDate!) * factor);
         this.updateTimelineXScale();
         this.triggerCallbacks(); // trigger callbacks to update previously rendered timeline
     }
@@ -433,8 +436,10 @@ export const calculateAverageSessionFrequencyHours = (handler: TestDataHandler,
         for (let i = 0; i < Math.min(sessions.length, thresholdNumberOfSessions); i++) {
             const commitId = sessions[i].commitId;
             if (!seenCommitIds.has(commitId)) {
+                const commitDate = handler.commitIdDates.get(commitId);
+                if (!commitDate) continue; // no matching commit date found? should not happen
                 seenCommitIds.add(commitId);
-                uniqueDates.push(handler.commitIdDates.get(commitId)!.getTime() / 1000);
+                uniqueDates.push(commitDate.getTime() / 1000);
             }
         }
 
@@ -448,3 +453,35 @@ export const calculateAverageSessionFrequencyHours = (handler: TestDataHandler,
 
     return averageSessionFrequencyHours;
 }
+
+export const getSessionsNearTimeFactory = (sessions: TestSessionResult[] | PhaseSessionResult[], sessionTimes: number[], days: number, metaCount: number = 1) => {
+        // Organize sessions into timeline sections by sessionTime, to enable fast lookup of sessions near a specific time
+        // We'll divide the entire timeline into bins of fixed size.
+        const sectionSizeSeconds = Math.min(2, (days * 24 / 2) / (sessions.length / metaCount)) * 60 * 60; // 2 hours per section bin minimum, otherwise inversely proportional to session count
+        const sessionSections: Map<number, any> = new Map();
+        // Time reference for binning
+        const minSessionTime = sessionTimes.length > 0 ? sessionTimes[sessionTimes.length - 1] : 0;
+        // Organize sessions into bins
+        sessions.forEach((session, idx) => {
+            const sessionTime = sessionTimes[idx];
+            const sectionIdx = Math.floor((sessionTime - minSessionTime) / sectionSizeSeconds);
+            if (!sessionSections.has(sectionIdx)) {
+                sessionSections.set(sectionIdx, []);
+            }
+            sessionSections.get(sectionIdx)!.push(session);
+        });
+
+        function getSessionsNearTime(targetTime: number, windowSections: number = 1): TestSessionResult[] | PhaseSessionResult[] {
+            const sectionIdx = Math.floor((targetTime - minSessionTime) / sectionSizeSeconds);
+            const results: TestSessionResult[] | PhaseSessionResult[] = [];
+            const startSectionIdx = sectionIdx - windowSections;
+            const endSectionIdx = sectionIdx + windowSections;
+            for (let i = startSectionIdx; i <= endSectionIdx; ++i) {
+                if (sessionSections.has(i)) {
+                    results.push(...sessionSections.get(i)!);
+                }
+            }
+            return results;
+        }
+        return getSessionsNearTime;
+}
\ No newline at end of file
