Horde: Add configurable log level for optional input retrieval failures

#Synced-CL 0

diff --git a/Engine/Source/Programs/Horde/Drivers/JobDriver/Execution/JobExecutor.cs b/Engine/Source/Programs/Horde/Drivers/JobDriver/Execution/JobExecutor.cs
--- a/Engine/Source/Programs/Horde/Drivers/JobDriver/Execution/JobExecutor.cs
+++ b/Engine/Source/Programs/Horde/Drivers/JobDriver/Execution/JobExecutor.cs
@@ -1472,7 +1472,14 @@ namespace JobDriver.Execution
 				}
 				catch (Exception ex) when (ex is not OperationCanceledException)
 				{
-					logger.LogWarning(ex, "Failed to read optional input {OptionalInput} from temp storage: {Message}. The producing step may have failed. Continuing without it", optionalInput, ex.Message);
+					if (Environment.GetEnvironmentVariable("UE_HORDE_LOG_OPTIONAL_INPUT_FAILURE_AS_WARNING") == "1")
+					{
+						logger.LogWarning(ex, "Failed to read optional input {OptionalInput} from temp storage: {Message}. The producing step may have failed. Continuing without it", optionalInput, ex.Message);
+					}
+					else
+					{
+						logger.LogInformation(ex, "Failed to read optional input {OptionalInput} from temp storage: {Message}. The producing step may have failed. Continuing without it", optionalInput, ex.Message);
+					}
 					continue;
 				}
 				tagNameToFileSet[tagName] = fileList.ToFileSet(unrealRootDir);
diff --git a/Engine/Source/Programs/Horde/HordeAgent/AgentSettings.cs b/Engine/Source/Programs/Horde/HordeAgent/AgentSettings.cs
--- a/Engine/Source/Programs/Horde/HordeAgent/AgentSettings.cs
+++ b/Engine/Source/Programs/Horde/HordeAgent/AgentSettings.cs
@@ -278,6 +278,12 @@ namespace HordeAgent
 		/// </summary>
 		public int? TempStorageMaxBatchSize { get; set; } = null;
 
+		/// <summary>
+		/// Whether to log optional input retrieval failures as warnings instead of information messages.
+		/// Defaults to false (information level).
+		/// </summary>
+		public bool LogOptionalInputFailureAsWarning { get; set; } = false;
+
 		/// <summary>
 		/// Gets the current server settings
 		/// </summary>
diff --git a/Engine/Source/Programs/Horde/HordeAgent/Leases/Handlers/JobHandler.cs b/Engine/Source/Programs/Horde/HordeAgent/Leases/Handlers/JobHandler.cs
--- a/Engine/Source/Programs/Horde/HordeAgent/Leases/Handlers/JobHandler.cs
+++ b/Engine/Source/Programs/Horde/HordeAgent/Leases/Handlers/JobHandler.cs
@@ -147,6 +147,7 @@ namespace HordeAgent.Leases.Handlers
 				{
 					environment["UE_HORDE_TEMP_STORAGE_MAX_BATCH_SIZE"] = Convert.ToString(_settings.CurrentValue.TempStorageMaxBatchSize.Value, CultureInfo.InvariantCulture);
 				}
+				environment["UE_HORDE_LOG_OPTIONAL_INPUT_FAILURE_AS_WARNING"] = _settings.CurrentValue.LogOptionalInputFailureAsWarning ? "1" : "0";
 
 				if (_driverCancelEventName != null)
 				{
