[Test Automation Hub] When loading test report, only load stream data in the background for efficiency.

#rnx
#Synced-CL 52296502

diff --git a/Engine/Source/Programs/Horde/HordeDashboard/plugins/test-automation/testAutomationPhaseView.tsx b/Engine/Source/Programs/Horde/HordeDashboard/plugins/test-automation/testAutomationPhaseView.tsx
--- a/Engine/Source/Programs/Horde/HordeDashboard/plugins/test-automation/testAutomationPhaseView.tsx
+++ b/Engine/Source/Programs/Horde/HordeDashboard/plugins/test-automation/testAutomationPhaseView.tsx
@@ -27,7 +27,14 @@ const PhaseHistoryPane: React.FC<{ test: TestNameRef, meta?: MetadataRef, phaseK
                 <Stack style={{ marginLeft: 30, marginBottom: 18, width: 10, cursor: "default" }} title="History">
                     <FontIcon style={{ fontSize: 15 }} iconName="Timeline" />
                 </Stack>
-                {(!metaStatus || !sessionId) && <Text style={{ marginLeft: 10,marginBottom: 21}}>No history available for this phase</Text>}
+                {(!metaStatus || !sessionId) &&
+                    <Stack style={{ marginLeft: 10,marginBottom: 21}}>
+                        {handler.queryLoading ?
+                            <Stack horizontal verticalAlign="center" tokens={{childrenGap: 6}}><Spinner size={SpinnerSize.small} /><Text>Loading...</Text></Stack>
+                            : <Text>No history available for this phase</Text>
+                        }
+                    </Stack>
+                }
                 {!!metaStatus && !!sessionId && <PhaseHistoryGraphWidget
                     test={test}
                     streamId={streamId!}
@@ -52,7 +59,7 @@ const PhaseHistoryPane: React.FC<{ test: TestNameRef, meta?: MetadataRef, phaseK
             </Stack>
 });
 
-const TestPhaseExpandPane: React.FC<{ phase: TestPhaseStatus, onEventsLoaded?: (events: EventEntry[]) => void, handler: TestDataHandler }> = memo(({ phase, onEventsLoaded, handler }) => {
+const TestPhaseExpandPane: React.FC<{ phase: TestPhaseStatus, onEventsLoaded?: (events: EventEntry[]) => void, handler: TestDataHandler }> = memo(observer(({ phase, onEventsLoaded, handler }) => {
     const [sessionId, setSessionId] = useState<string | undefined>();
     const [comparison, setComparison] = useState<boolean>(comparisonContext.isEnable);
 
@@ -60,15 +67,15 @@ const TestPhaseExpandPane: React.FC<{ phase: TestPhaseStatus, onEventsLoaded?: (
         // query phase data
         setSessionId(undefined);
         handler.queryPhaseSessions([phase], phase.testSession?.streamId ?? handler.stream!).then(() => {
-            setComparison(comparisonContext.isEnable);
-            setSessionId(phase.session?.id ?? "");
-            if (!phase.session) {
+            if (!comparisonContext.isEnable) {
                 // if no phase session associated with this phase in history, it should be a preflight.
                 // we are going to prefill the commitId in the comparison context with the preflight base commit id.
-                comparisonContext.commitId = phase.commitId;
+                comparisonContext.commitId = !phase.session ? phase.commitId : undefined;
             }
+            setComparison(comparisonContext.isEnable);
+            setSessionId(phase.session?.id ?? "");
         });
-    }, [phase]);
+    }, [phase, phase.testSession]);
 
     return <Stack style={{overflow: 'hidden'}} grow>
                 {sessionId === undefined &&
@@ -109,7 +116,7 @@ const TestPhaseExpandPane: React.FC<{ phase: TestPhaseStatus, onEventsLoaded?: (
                     </Stack>
                 }
             </Stack>
-});
+}));
 
 const TestPhaseHeader: React.FC<{ phase: TestPhaseStatus, expanded: boolean, toggleView: (phaseKey: string) => void, onNext?: () => void, onBack?: () => void }> = memo(({ phase, expanded, toggleView, onNext, onBack }) => {
     const onClick = useCallback(() => {
@@ -290,7 +297,6 @@ export const TestPhasesView: React.FC<{ testDataId: string, handler: TestDataHan
     const warningCursor = useRef<CircularList<string>>();
     const phaseCursor = useRef<CircularList<string>>();
     const listRef = useRef<List>();
-    const jobContext = useRef<JobContext>();
 
     const phaseKey = handler.getSearchParam('phase') as string | undefined;
 
@@ -340,12 +346,13 @@ export const TestPhasesView: React.FC<{ testDataId: string, handler: TestDataHan
                 setSessionDetails(fetchedSessionDetails);
                 const testId = fetchedSessionDetails.test?.id ?? testDataId;
                 updateFilter(phaseFiltersContext.filterTestPhases(testId, fetchedSessionDetails.phases) as TestPhaseStatus[]);
-                jobContext.current = getJobParams(fetchedSessionDetails, handler);
             }
         });
         setSessionDetails(undefined);
     }, [testDataId]);
 
+    const jobContext : JobContext | undefined = useMemo(() => sessionDetails ? getJobParams(sessionDetails, handler) : undefined, [sessionDetails?.session]);
+
     const onPagesUpdated = useCallback(() => {
         if (!focusPhase) return;
         const phaseIndex = filteredPhases.findIndex((phase) => phase.key === focusPhase);
@@ -401,9 +408,9 @@ export const TestPhasesView: React.FC<{ testDataId: string, handler: TestDataHan
     const testId = sessionDetails?.test?.id ?? testDataId;
 
     return <Stack style={{width: '100%', height: '100%'}}>
-                {!!jobArtifacts && jobContext.current?.job?.id && jobContext.current?.step?.id &&
+                {!!jobArtifacts && jobContext?.job?.id && jobContext?.step?.id &&
                     <JobArtifactsModal
-                        jobId={jobContext.current.job.id} stepId={jobContext.current?.step?.id} contextType={jobArtifacts.type} artifacts={[jobArtifacts]}
+                        jobId={jobContext.job.id} stepId={jobContext?.step?.id} contextType={jobArtifacts.type} artifacts={[jobArtifacts]}
                         onClose={() => {
                             setJobArtifacts(undefined);
                             handler.removeSearchParam('artifactPath');
@@ -439,8 +446,8 @@ export const TestPhasesView: React.FC<{ testDataId: string, handler: TestDataHan
                                 <Text style={{ fontSize: 16, fontWeight: 600, whiteSpace: 'nowrap', overflow: 'hidden', textOverflow: 'ellipsis' }}>{`${sessionDetails.name} - ${metaString}`}</Text>
                                 <Stack style={{ paddingLeft: 8 }} horizontal tokens={{childrenGap: 8}}>
                                     <Text>{projectStore.streamById(sessionDetails.streamId)?.fullname ?? sessionDetails.streamId}</Text>
-                                    <ChangeButton job={jobContext.current?.job} rangeCL={jobContext.current?.rangeCL} />
-                                    <Link href={`/job/${jobContext.current?.job?.id}?step=${jobContext.current?.step?.id}`} target="_blank"><Text style={{ fontSize: 12, whiteSpace: 'nowrap', color: hordeTheme.semanticColors.link }}>View Job</Text></Link>
+                                    <ChangeButton job={jobContext?.job} rangeCL={jobContext?.rangeCL} />
+                                    <Link href={`/job/${jobContext?.job?.id}?step=${jobContext?.step?.id}`} target="_blank"><Text style={{ fontSize: 12, whiteSpace: 'nowrap', color: hordeTheme.semanticColors.link }}>View Job</Text></Link>
                                     <Text style={{whiteSpace: 'nowrap'}}><span style={styles.labelSmall}>Run on: </span><span>{getShortNiceTime(sessionDetails.start, true, true, true)}</span></Text>
                                     <Text style={{whiteSpace: 'nowrap'}}><span style={styles.labelSmall}> For: </span><span>{msecToElapsed(sessionDetails.duration * 1000, true, true)}</span></Text>
                                     <Stack horizontalAlign="end">
diff --git a/Engine/Source/Programs/Horde/HordeDashboard/plugins/test-automation/testData.ts b/Engine/Source/Programs/Horde/HordeDashboard/plugins/test-automation/testData.ts
--- a/Engine/Source/Programs/Horde/HordeDashboard/plugins/test-automation/testData.ts
+++ b/Engine/Source/Programs/Horde/HordeDashboard/plugins/test-automation/testData.ts
@@ -406,6 +406,7 @@ export class TestSessionDetails {
         this.stepId = stepId;
         this.devices = new Map();
         this.phases = [];
+        makeObservable(this);
     }
 
     get start() {
@@ -463,11 +464,13 @@ export class TestSessionDetails {
     
     test?: TestNameRef;
     meta?: MetadataRef;
-    session?: TestSessionResult;
     tags?: string[];
     preflightCommitId?: string;
     commandline?: string;
 
+    @observable
+    session?: TestSessionResult;
+
     private _start: Date;
     private _duration: number;
     private _artifactFactory?: ArtifactFactory;
@@ -1663,10 +1666,11 @@ export class TestDataHandler {
             try {
                 const sessionDetails = await fetcherVersion.fetchTestData(testDataId);
                 if (sessionDetails) {
-                    await this.query(sessionDetails.streamId);
-                    const session = this.getStatusStream(sessionDetails.streamId)?.tests.get(sessionDetails.test!)?.sessions.get(sessionDetails.meta!)?.history.find(s => s.testDataId === testDataId);
-                    sessionDetails.session = session;
                     this.testTestData.set(testDataId, sessionDetails)
+                    this.query(sessionDetails.streamId).then(() => {
+                        const session = this.getStatusStream(sessionDetails.streamId)?.tests.get(sessionDetails.test!)?.sessions.get(sessionDetails.meta!)?.history.find(s => s.testDataId === testDataId);
+                        sessionDetails.session = session;
+                    });
                     /// stop fetching if a valid object was pulled
                     break;
                 }
