Horde: Fix cross-batch dependency resolution in RefreshDependentJobSteps

Split stepForNode construction into separate first pass before dependency evaluation

#Synced-CL 0

diff --git a/Engine/Source/Programs/Horde/Plugins/Build/HordeServer.Build.Tests/Jobs/JobCollectionTests.cs b/Engine/Source/Programs/Horde/Plugins/Build/HordeServer.Build.Tests/Jobs/JobCollectionTests.cs
--- a/Engine/Source/Programs/Horde/Plugins/Build/HordeServer.Build.Tests/Jobs/JobCollectionTests.cs
+++ b/Engine/Source/Programs/Horde/Plugins/Build/HordeServer.Build.Tests/Jobs/JobCollectionTests.cs
@@ -1796,5 +1796,74 @@ namespace HordeServer.Tests.Jobs
 			Assert.AreEqual(JobStepBatchError.NoLongerNeeded, job.Batches[0].Error, "Old batch should be marked NoLongerNeeded when step is retried");
 			Assert.AreEqual(JobStepBatchError.None, job.Batches[1].Error, "New batch should have no error");
 		}
+
+		[TestMethod]
+		public async Task UpdateGraphWhileBatchRunningPreservesDependenciesAsync()
+		{
+			// Regression test: When UpdateGraphForHorde runs while a batch is already executing
+			// (due to RunEarly), nodes added to that batch with dependencies on nodes in new
+			// (later) batches must stay in Waiting state. Previously, RefreshDependentJobSteps
+			// built the stepForNode dictionary incrementally during iteration, so dependencies
+			// in later batches were invisible and the step was incorrectly marked Ready.
+
+			Mock<ITemplate> templateMock = new(MockBehavior.Strict);
+			templateMock.SetupGet(x => x.InitialAgentType).Returns((string?)null);
+
+			IGraph baseGraph = await GraphCollection.AddAsync(templateMock.Object);
+
+			// Initial graph: single group with a RunEarly node
+			List<NewGroup> initialGroups = [];
+			NewGroup cookGroup = AddGroup(initialGroups);
+			AddNode(cookGroup, "Cook", null, n => { n.RunEarly = true; });
+
+			IGraph graph1 = await GraphCollection.AppendAsync(baseGraph, initialGroups);
+
+			CreateJobOptions options = new CreateJobOptions();
+			options.Arguments.Add("-Target=Cook");
+			options.Arguments.Add("-Target=Download Artifacts");
+
+			IJob job = await JobCollection.AddAsync(JobIdUtils.GenerateNewId(), _streamId, _templateId1, ContentHash.SHA1("hello"), baseGraph, "Test job", CommitIdWithOrder.FromPerforceChange(123), CommitIdWithOrder.FromPerforceChange(123), options);
+
+			job = await RunBatchAsync(job, 0); // Setup Build
+			job = await RunStepAsync(job, 0, 0, JobStepOutcome.Success); // Setup Build
+
+			job = Deref(await job.TryUpdateGraphAsync(graph1));
+
+			// Cook batch should be Ready due to RunEarly (Setup Build + Cook)
+			Assert.HasCount(2, job.Batches);
+			Assert.AreEqual(JobStepBatchState.Ready, job.Batches[1].State); // Cook batch
+
+			// Start the Cook batch and step (simulating RunEarly agent pickup)
+			job = await RunBatchAsync(job, 1);
+			job = Deref(await job.TryUpdateStepAsync(job.Batches[1].Id, job.Batches[1].Steps[0].Id, JobStepState.Running, JobStepOutcome.Success));
+			Assert.AreEqual(JobStepBatchState.Running, job.Batches[1].State);
+			Assert.AreEqual(JobStepState.Running, job.Batches[1].Steps[0].State);
+
+			// Now build an updated graph that adds a new node in a new group (Upload Artifacts)
+			// and a dependent node in Cook's group. The Upload Artifacts group must be declared
+			// first so the graph can resolve the dependency reference; but at runtime the
+			// new batch for Upload Artifacts will be appended AFTER the existing Running batch,
+			// which is the core of the bug.
+			List<NewGroup> updatedGroups = [];
+
+			NewGroup uploadGroup = AddGroup(updatedGroups);
+			AddNode(uploadGroup, "Upload Artifacts", null);
+
+			NewGroup cookGroupUpdated = AddGroup(updatedGroups);
+			AddNode(cookGroupUpdated, "Cook", null, n => { n.RunEarly = true; });
+			AddNode(cookGroupUpdated, "Download Artifacts", ["Upload Artifacts"]);
+
+			IGraph graph2 = Deref(await GraphCollection.AppendAsync(baseGraph, updatedGroups));
+
+			// Update the graph while the Cook batch is Running
+			job = Deref(await job.TryUpdateGraphAsync(graph2));
+
+			Assert.HasCount(3, job.Batches);
+			// Upload Artifacts has no dependencies and should be Ready.
+			Assert.AreEqual(JobStepState.Ready, job.Batches[2].Steps[0].State);
+			// Download Artifacts should remain Waiting until Upload Artifacts completes.
+			// If this is Ready, the dependency was silently dropped during graph update.
+			Assert.AreEqual(JobStepState.Waiting, job.Batches[1].Steps[1].State);
+		}
 	}
 }
diff --git a/Engine/Source/Programs/Horde/Plugins/Build/HordeServer.Build/Jobs/JobCollection.cs b/Engine/Source/Programs/Horde/Plugins/Build/HordeServer.Build/Jobs/JobCollection.cs
--- a/Engine/Source/Programs/Horde/Plugins/Build/HordeServer.Build/Jobs/JobCollection.cs
+++ b/Engine/Source/Programs/Horde/Plugins/Build/HordeServer.Build/Jobs/JobCollection.cs
@@ -2693,7 +2693,21 @@ namespace HordeServer.Jobs
 			UpdateDefinitionBuilder<JobDocument> updateBuilder = Builders<JobDocument>.Update;
 			if (job.Batches != null)
 			{
+				// First pass: build a complete mapping of nodes to steps so that dependency lookups
+				// are not affected by the order batches appear in the list. This fixes a bug where
+				// steps in earlier batches could have their dependencies missed when those dependencies
+				// were in batches appended later (e.g. during a graph update while a batch is Running).
 				Dictionary<INode, JobStepDocument> stepForNode = new Dictionary<INode, JobStepDocument>();
+				foreach (JobStepBatchDocument batch in job.Batches)
+				{
+					foreach (JobStepDocument step in batch.Steps)
+					{
+						INode node = graph.Groups[batch.GroupIdx].Nodes[step.NodeIdx];
+						stepForNode[node] = step;
+					}
+				}
+
+				// Second pass: evaluate dependencies and update step/batch states
 				for (int loopBatchIdx = 0; loopBatchIdx < job.Batches.Count; loopBatchIdx++)
 				{
 					int batchIdx = loopBatchIdx; // For lambda capture
@@ -2768,8 +2782,6 @@ namespace HordeServer.Jobs
 							updates.Add(updateBuilder.Set(x => x.Batches[batchIdx].Steps[stepIdx].Outcome, newOutcome));
 							JobAudit(job.Id, logger, jobAudit, "Setting step {StepId} outcome to {StepOutcome}", step.Id, newOutcome);
 						}
-
-						stepForNode[node] = step;
 					}
 
 					if (batch.State == JobStepBatchState.Waiting || batch.State == JobStepBatchState.Ready)
