Horde: Update comments about skipped optional inputs

#rnx
#Synced-CL 0

diff --git a/Engine/Source/Programs/Horde/Plugins/Build/HordeServer.Build/Jobs/JobCollection.cs b/Engine/Source/Programs/Horde/Plugins/Build/HordeServer.Build/Jobs/JobCollection.cs
--- a/Engine/Source/Programs/Horde/Plugins/Build/HordeServer.Build/Jobs/JobCollection.cs
+++ b/Engine/Source/Programs/Horde/Plugins/Build/HordeServer.Build/Jobs/JobCollection.cs
@@ -2731,6 +2731,9 @@ namespace HordeServer.Jobs
 							// An optional input step aborted by the driver (not user requested) still allows the step to move to ready.
 							// A skipped optional input step does not allow continuation; only allow continuation when the dependency ran to
 							// completion and failed, not when it was skipped entirely.
+							// Check all dependencies for skipped, including optional ones.
+							// A skipped optional dependency means its required subgraph failed, which is different from ran and failed.
+							// An optional dependency that ran and failed is tolerated but one that was skipped entirely is not.
 							if (steps.Any(x => x.AbortRequested || x.State == JobStepState.Skipped) || requiredSteps.Any(x => x.Outcome == JobStepOutcome.Failure))
 							{
 								newState = JobStepState.Skipped;
