Horde: Skip unnecessary clean in PerforceMaterializer during conforms on first boot

diff --git a/Engine/Source/Programs/Horde/Drivers/JobDriver.Tests/PerforceMaterializerTests.cs b/Engine/Source/Programs/Horde/Drivers/JobDriver.Tests/PerforceMaterializerTests.cs
--- a/Engine/Source/Programs/Horde/Drivers/JobDriver.Tests/PerforceMaterializerTests.cs
+++ b/Engine/Source/Programs/Horde/Drivers/JobDriver.Tests/PerforceMaterializerTests.cs
@@ -33,6 +33,26 @@ public class TestLogger : ILogger
 	public IDisposable? BeginScope<TState>(TState state) where TState : notnull => null;
 }
 
+/// <summary>
+/// Logger wrapper that captures formatted messages while delegating to an inner logger.
+/// </summary>
+public class LogCapture : ILogger
+{
+	private readonly ILogger _inner;
+	public List<string> Messages { get; } = [];
+
+	public LogCapture(ILogger inner) => _inner = inner;
+
+	public void Log<TState>(LogLevel logLevel, EventId eventId, TState state, Exception? exception, Func<TState, Exception?, string> formatter)
+	{
+		Messages.Add(formatter(state, exception));
+		_inner.Log(logLevel, eventId, state, exception, formatter);
+	}
+
+	public bool IsEnabled(LogLevel logLevel) => _inner.IsEnabled(logLevel);
+	public IDisposable? BeginScope<TState>(TState state) where TState : notnull => _inner.BeginScope(state);
+}
+
 [TestClass]
 public class PerforceMaterializerTests : BasePerforceFixtureTest
 {
@@ -240,7 +260,7 @@ public class PerforceMaterializerTests : BasePerforceFixtureTest
 	}
 
 	[TestMethod]
-	public async Task Sync_WithChangedClientName_ShouldForceCleanResync_Async()
+	public async Task Sync_WithChangedClientName_ShouldFlushNotWipe_Async()
 	{
 		CancellationToken ct = CancellationToken.None;
 		using PerforceMaterializer pm = CreatePm(PerforceSnapshot.FooMain);
@@ -253,7 +273,7 @@ public class PerforceMaterializerTests : BasePerforceFixtureTest
 		// Step 2: Write a marker file into the workspace directory
 		string markerPath = Path.Join(pm.SyncDir.FullName, "_edge_test_marker.txt");
 		await File.WriteAllTextAsync(markerPath, "marker", ct);
-		Assert.IsTrue(File.Exists(markerPath), "Marker file should exist before re-sync");
+		Assert.IsTrue(File.Exists(markerPath));
 
 		// Step 3: Tamper with the client name in State.json to simulate an edge server switch
 		PerforceMaterializer.State? state = await pm.LoadStateForTestAsync(ct);
@@ -262,12 +282,96 @@ public class PerforceMaterializerTests : BasePerforceFixtureTest
 		string stateFilePath = Path.Join(pm.BaseDir.FullName, "State.json");
 		await File.WriteAllTextAsync(stateFilePath, JsonSerializer.Serialize(state, new JsonSerializerOptions { WriteIndented = true }), ct);
 
-		// Step 4: Sync to CL 5 — the changed client name should trigger isDirty and wipe the workspace
+		// Step 4: Sync to CL 5  - edge server switch should flush, not wipe
+		await pm.SyncAsync(5, -1, new SyncOptions { RemoveUntracked = false }, ct);
+
+		// Step 5: Assert the marker file still exists (workspace was NOT wiped, only flushed)
+		Assert.IsTrue(File.Exists(markerPath));
+	}
+
+	[TestMethod]
+	public async Task Sync_WithChangedClientNameAndStream_ShouldWipe_Async()
+	{
+		CancellationToken ct = CancellationToken.None;
+		using PerforceMaterializer pm = CreatePm(PerforceSnapshot.FooMain);
+
+		// Step 1: Sync to CL 2 and finalize so state is saved to disk
+		await pm.SyncAsync(2, -1, new SyncOptions { RemoveUntracked = false }, ct);
+		PerforceSnapshot.GetSnapshot().GetStream(PerforceSnapshot.FooMain).GetWorkspace(2).AssertDepotFiles(pm.SyncDir.FullName);
+		await pm.FinalizeAsync(ct);
+
+		// Step 2: Write a marker file into the workspace directory
+		string markerPath = Path.Join(pm.SyncDir.FullName, "_edge_test_marker.txt");
+		await File.WriteAllTextAsync(markerPath, "marker", ct);
+		Assert.IsTrue(File.Exists(markerPath));
+
+		// Step 3: Tamper with both client name AND stream to simulate a genuine workspace change
+		PerforceMaterializer.State? state = await pm.LoadStateForTestAsync(ct);
+		Assert.IsNotNull(state);
+		state = state with { Client = "Horde+PM+FAKE+id+fakeEdge", Stream = "//Foo/OtherStream" };
+		string stateFilePath = Path.Join(pm.BaseDir.FullName, "State.json");
+		await File.WriteAllTextAsync(stateFilePath, JsonSerializer.Serialize(state, new JsonSerializerOptions { WriteIndented = true }), ct);
+
+		// Step 4: Sync  - the combined changes should trigger isDirty and wipe the workspace
 		await pm.SyncAsync(5, -1, new SyncOptions { RemoveUntracked = false }, ct);
 		PerforceSnapshot.GetSnapshot().GetStream(PerforceSnapshot.FooMain).GetWorkspace(5).AssertDepotFiles(pm.SyncDir.FullName);
 
 		// Step 5: Assert the marker file was deleted (workspace was wiped because isDirty was true)
-		Assert.IsFalse(File.Exists(markerPath), "Marker file should have been deleted by workspace wipe triggered by client name change");
+		Assert.IsFalse(File.Exists(markerPath));
+	}
+
+	[TestMethod]
+	[DataRow(true, DisplayName = "Fresh sync")]
+	[DataRow(false, DisplayName = "Subsequent sync")]
+	public async Task Finalize_AlwaysRunsClean_Async(bool isFreshSync)
+	{
+		// FinalizeAsync should always run CleanAsync when RemoveUntracked=true,
+		// even on a fresh sync, because the job/build step runs between SyncAsync and FinalizeAsync and may modify the workspace.
+		LogCapture log = new(LoggerFactory.CreateLogger("PM"));
+		using PerforceMaterializer pm = CreatePm(PerforceSnapshot.FooMain, logger: log);
+
+		if (!isFreshSync)
+		{
+			await pm.SyncAsync(2, -1, new SyncOptions { RemoveUntracked = true }, CancellationToken.None);
+			await pm.FinalizeAsync(CancellationToken.None);
+		}
+
+		await pm.SyncAsync(5, -1, new SyncOptions { RemoveUntracked = true }, CancellationToken.None);
+		log.Messages.Clear();
+
+		await pm.FinalizeAsync(CancellationToken.None);
+
+		Assert.IsTrue(log.Messages.Any(m => m.Contains("Cleaning files", StringComparison.Ordinal)));
+	}
+
+	[TestMethod]
+	public async Task Conform_FreshState_SkipsClean_Async()
+	{
+		// ConformAsync with no prior state (fresh) should skip CleanAsync even when removeUntrackedFiles=true.
+		LogCapture log = new(LoggerFactory.CreateLogger("PM"));
+		using PerforceMaterializer pm = CreatePm(PerforceSnapshot.FooMain, logger: log);
+		log.Messages.Clear();
+
+		await pm.ConformAsync(true, CancellationToken.None);
+
+		Assert.IsFalse(log.Messages.Any(m => m.Contains("Cleaning files", StringComparison.Ordinal)));
+	}
+
+	[TestMethod]
+	public async Task Conform_ExistingState_RunsClean_Async()
+	{
+		// ConformAsync with existing state should run CleanAsync when removeUntrackedFiles=true.
+		LogCapture log = new(LoggerFactory.CreateLogger("PM"));
+		using PerforceMaterializer pm = CreatePm(PerforceSnapshot.FooMain, logger: log);
+
+		// Establish state with a first sync+finalize
+		await pm.SyncAsync(2, -1, new SyncOptions(), CancellationToken.None);
+		await pm.FinalizeAsync(CancellationToken.None);
+		log.Messages.Clear();
+
+		await pm.ConformAsync(true, CancellationToken.None);
+
+		Assert.IsTrue(log.Messages.Any(m => m.Contains("Cleaning files", StringComparison.Ordinal)));
 	}
 
 	[TestMethod]
@@ -386,9 +490,9 @@ public class PerforceMaterializerTests : BasePerforceFixtureTest
 		Assert.AreEqual(PerforceMaterializer.TransactionStatus.Clean, pmState.Status);
 	}
 
-	private PerforceMaterializer CreatePm(string stream, string? dataDir = null, string? identifier = "pmTest", string[]? view = null)
+	private PerforceMaterializer CreatePm(string stream, string? dataDir = null, string? identifier = "pmTest", string[]? view = null, ILogger? logger = null)
 	{
-		ILogger logger = LoggerFactory.CreateLogger("PM");
+		logger ??= LoggerFactory.CreateLogger("PM");
 		RpcAgentWorkspace raw = new ()
 		{
 			ServerAndPort = PerforceConnection.ServerAndPort,
diff --git a/Engine/Source/Programs/Horde/Drivers/JobDriver/Execution/PerforceMaterializer.cs b/Engine/Source/Programs/Horde/Drivers/JobDriver/Execution/PerforceMaterializer.cs
--- a/Engine/Source/Programs/Horde/Drivers/JobDriver/Execution/PerforceMaterializer.cs
+++ b/Engine/Source/Programs/Horde/Drivers/JobDriver/Execution/PerforceMaterializer.cs
@@ -170,6 +170,7 @@ public sealed class PerforceMaterializer : IWorkspaceMaterializer
 	private List<string>? _streamView;
 	private int? _lastSyncedChangeNum;
 	private SyncOptions? _lastSyncOptions;
+    private bool _isFreshSync;
 	private const string StateFilename = "State.json";
 	private const bool PreferNativePerforceClient = true;
 	private const int MaxFilesLogged = 500;
@@ -280,7 +281,8 @@ public sealed class PerforceMaterializer : IWorkspaceMaterializer
 
 		_lastSyncOptions = options;
 		State? state = await LoadStateAsync(cancellationToken);
-		
+		_isFreshSync = state == null;
+
 		// Check if only the client name changed (edge server switch scenario)
 		// In this case, don't mark as dirty - the existing flush + sync will handle it correctly
 		bool isClientNameChange = state != null
@@ -432,7 +434,7 @@ public sealed class PerforceMaterializer : IWorkspaceMaterializer
 		SyncOptions syncOptions = new () { RemoveUntracked = removeUntrackedFiles };
 		await SyncAsync(IWorkspaceMaterializer.LatestChangeNumber, 0, syncOptions, cancellationToken);
 		
-		if (removeUntrackedFiles)
+		if (removeUntrackedFiles && !_isFreshSync)
 		{
 			if (_perforceWithClient != null)
 			{
