[Test Automation Hub] Fix preflight job link and commit widget

#rnx
#Synced-CL 52381418

diff --git a/Engine/Source/Programs/Horde/HordeDashboard/plugins/test-automation/testAutomationPhaseView.tsx b/Engine/Source/Programs/Horde/HordeDashboard/plugins/test-automation/testAutomationPhaseView.tsx
--- a/Engine/Source/Programs/Horde/HordeDashboard/plugins/test-automation/testAutomationPhaseView.tsx
+++ b/Engine/Source/Programs/Horde/HordeDashboard/plugins/test-automation/testAutomationPhaseView.tsx
@@ -18,7 +18,7 @@ import { getEventStyling, PhaseEventsPane } from "./testAutomationEvents";
 import { comparisonContext, PhaseComparisonPane } from "./testAutomationPhaseCompare";
 import { EventEntry } from "./testEventsModel";
 
-const PhaseHistoryPane: React.FC<{ test: TestNameRef, meta?: MetadataRef, phaseKey: string, streamId?: string, sessionId?: string, handler: TestDataHandler, onClickCompare: () => void }> = memo(({test, meta, phaseKey, streamId, sessionId, handler, onClickCompare}) => {
+const PhaseHistoryPane: React.FC<{ test: TestNameRef, meta?: MetadataRef, phaseKey: string, streamId?: string, sessionId?: string, handler: TestDataHandler, onClickCompare: () => void }> = memo(observer(({test, meta, phaseKey, streamId, sessionId, handler, onClickCompare}) => {
 
     const metaStatus = !!meta && !!streamId? handler.getStatusStream(streamId)?.tests.get(test)?.sessions.get(meta) : undefined;
     const isFromGridView = !!handler.getSearchParam('grid');
@@ -31,7 +31,7 @@ const PhaseHistoryPane: React.FC<{ test: TestNameRef, meta?: MetadataRef, phaseK
                     <Stack style={{ marginLeft: 10,marginBottom: 21}}>
                         {handler.queryLoading ?
                             <Stack horizontal verticalAlign="center" tokens={{childrenGap: 6}}><Spinner size={SpinnerSize.small} /><Text>Loading...</Text></Stack>
-                            : <Text>No history available for this phase</Text>
+                            : <Text>No history available for this phase (for preflight, open the compare view)</Text>
                         }
                     </Stack>
                 }
@@ -57,7 +57,7 @@ const PhaseHistoryPane: React.FC<{ test: TestNameRef, meta?: MetadataRef, phaseK
                     />
                 </Stack>
             </Stack>
-});
+}));
 
 const TestPhaseExpandPane: React.FC<{ phase: TestPhaseStatus, onEventsLoaded?: (events: EventEntry[]) => void, handler: TestDataHandler }> = memo(observer(({ phase, onEventsLoaded, handler }) => {
     const [sessionId, setSessionId] = useState<string | undefined>();
@@ -351,7 +351,7 @@ export const TestPhasesView: React.FC<{ testDataId: string, handler: TestDataHan
         setSessionDetails(undefined);
     }, [testDataId]);
 
-    const jobContext : JobContext | undefined = useMemo(() => sessionDetails ? getJobParams(sessionDetails, handler) : undefined, [sessionDetails?.session]);
+    const jobContext : JobContext | undefined = useMemo(() => sessionDetails ? getJobParams(sessionDetails, handler) : undefined, [sessionDetails, sessionDetails?.session]);
 
     const onPagesUpdated = useCallback(() => {
         if (!focusPhase) return;
