[Horde][Analytics] Updating reference to underyling StructuredAnalytics proj.

#Synced-CL 52395043

diff --git a/Engine/Source/Programs/Horde/HordeServer/HordeServer.csproj b/Engine/Source/Programs/Horde/HordeServer/HordeServer.csproj
--- a/Engine/Source/Programs/Horde/HordeServer/HordeServer.csproj
+++ b/Engine/Source/Programs/Horde/HordeServer/HordeServer.csproj
@@ -107,10 +107,18 @@
     <ProjectReference Include="..\..\..\..\Restricted\NotForLicensees\Source\Programs\Horde\EpicPlugins\BuildHealth\HordeServer.Epic.BuildHealth\HordeServer.Epic.BuildHealth.csproj" />
   </ItemGroup>
 
+  <ItemGroup Condition="Exists('$(MSBuildThisFileDirectory)..\..\..\..\Restricted\NotForLicensees\Source\Programs\Horde\EpicPlugins\StructuredAnalytics\HordeServer.StructuredAnalytics\HordeServer.StructuredAnalytics.csproj')">
+    <ProjectReference Include="..\..\..\..\Restricted\NotForLicensees\Source\Programs\Horde\EpicPlugins\StructuredAnalytics\HordeServer.StructuredAnalytics\HordeServer.StructuredAnalytics.csproj" />
+  </ItemGroup>
+
   <ItemGroup Condition="Exists('$(MSBuildThisFileDirectory)..\..\..\..\Restricted\NotForLicensees\Source\Programs\Horde\EpicPlugins\EpicSandbox\HordeServer.Epic.EpicSandbox\HordeServer.Epic.EpicSandbox.csproj')">
     <ProjectReference Include="..\..\..\..\Restricted\NotForLicensees\Source\Programs\Horde\EpicPlugins\EpicSandbox\HordeServer.Epic.EpicSandbox\HordeServer.Epic.EpicSandbox.csproj" />
   </ItemGroup>
 
+  <ItemGroup Condition="Exists('$(MSBuildThisFileDirectory)..\..\..\..\Restricted\NotForLicensees\Source\Programs\Horde\EpicPlugins\ClickHouseTelemetry\HordeServer.Epic.ClickHouseTelemetry\HordeServer.Epic.ClickHouseTelemetry.csproj')">
+    <ProjectReference Include="..\..\..\..\Restricted\NotForLicensees\Source\Programs\Horde\EpicPlugins\ClickHouseTelemetry\HordeServer.Epic.ClickHouseTelemetry\HordeServer.Epic.ClickHouseTelemetry.csproj" />
+  </ItemGroup>
+
   <ItemGroup>
     <Content Include="..\ThirdParty\**" Condition="$([MSBuild]::IsOSPlatform('Windows'))">
       <CopyToOutputDirectory>Always</CopyToOutputDirectory>
