[Horde][Analytics][Server] Adding performance trends

#Synced-CL 52395043

diff --git a/Engine/Source/Programs/Horde/HordeServer/HordeServer.csproj b/Engine/Source/Programs/Horde/HordeServer/HordeServer.csproj
--- a/Engine/Source/Programs/Horde/HordeServer/HordeServer.csproj
+++ b/Engine/Source/Programs/Horde/HordeServer/HordeServer.csproj
@@ -118,6 +118,10 @@
   <ItemGroup Condition="Exists('$(MSBuildThisFileDirectory)..\..\..\..\Restricted\NotForLicensees\Source\Programs\Horde\EpicPlugins\ClickHouseTelemetry\HordeServer.Epic.ClickHouseTelemetry\HordeServer.Epic.ClickHouseTelemetry.csproj')">
     <ProjectReference Include="..\..\..\..\Restricted\NotForLicensees\Source\Programs\Horde\EpicPlugins\ClickHouseTelemetry\HordeServer.Epic.ClickHouseTelemetry\HordeServer.Epic.ClickHouseTelemetry.csproj" />
   </ItemGroup>
+  
+  <ItemGroup Condition="Exists('$(MSBuildThisFileDirectory)..\..\..\..\Restricted\NotForLicensees\Source\Programs\Horde\EpicPlugins\PerformanceTrends\HordeServer.PerformanceTrends\HordeServer.PerformanceTrends.csproj')">
+    <ProjectReference Include="..\..\..\..\Restricted\NotForLicensees\Source\Programs\Horde\EpicPlugins\PerformanceTrends\HordeServer.PerformanceTrends\HordeServer.PerformanceTrends.csproj" />
+  </ItemGroup>
 
   <ItemGroup>
     <Content Include="..\ThirdParty\**" Condition="$([MSBuild]::IsOSPlatform('Windows'))">
