JobDriver: AutomationTool.exe -> dll for modern path

#rnx
#Synced-CL 52421477

diff --git a/Engine/Source/Programs/Horde/Drivers/JobDriver/Execution/JobExecutor.cs b/Engine/Source/Programs/Horde/Drivers/JobDriver/Execution/JobExecutor.cs
--- a/Engine/Source/Programs/Horde/Drivers/JobDriver/Execution/JobExecutor.cs
+++ b/Engine/Source/Programs/Horde/Drivers/JobDriver/Execution/JobExecutor.cs
@@ -928,7 +928,7 @@ namespace JobDriver.Execution
 			FileReference[] automationToolPaths = new FileReference[]
 			{
 				FileReference.Combine(buildDir, $"{EnginePath}/Binaries/DotNET/AutomationTool.exe"),
-				FileReference.Combine(buildDir, $"{EnginePath}/Binaries/DotNET/AutomationTool/AutomationTool.exe")
+				FileReference.Combine(buildDir, $"{EnginePath}/Binaries/DotNET/AutomationTool/AutomationTool.dll")
 			};
 
 			if (automationToolPaths.Any(automationTool => FileReference.Exists(automationTool)))
