Horde: Updating schema documentation

#rnx

diff --git a/Engine/Source/Programs/Horde/Docs/Config/Schema/AclActions.md b/Engine/Source/Programs/Horde/Docs/Config/Schema/AclActions.md
--- a/Engine/Source/Programs/Horde/Docs/Config/Schema/AclActions.md
+++ b/Engine/Source/Programs/Horde/Docs/Config/Schema/AclActions.md
@@ -163,30 +163,6 @@
 | `ViewAgent` |  |
 | `ListAgents` |  |
 
-## Artifacts
-
-| Name | Description |
-| ---- | ----------- |
-| `ReadArtifact` | Permission to read from an artifact |
-| `WriteArtifact` | Permission to write to an artifact |
-| `DeleteArtifact` | Permission to delete to an artifact |
-| `UploadArtifact` | Ability to create an artifact. Typically just for debugging; agents have this access for a particular session. |
-| `DownloadArtifact` | Ability to download an artifact |
-
-## Bisect
-
-| Name | Description |
-| ---- | ----------- |
-| `CreateBisectTask` | Ability to start new bisect tasks |
-| `UpdateBisectTask` | Ability to update a bisect task |
-| `ViewBisectTask` | Ability to view a bisect task |
-
-## BuildHealth
-
-| Name | Description |
-| ---- | ----------- |
-| `AddUpdateDeleteBuildHealthFilter` | Ability to add, update, and delete build health filters. |
-
 ## Compute
 
 | Name | Description |
@@ -196,24 +172,6 @@
 | `UbaCacheRead` |  |
 | `UbaCacheWrite` |  |
 
-## Devices
-
-| Name | Description |
-| ---- | ----------- |
-| `DeviceRead` | Ability to read devices |
-| `DeviceWrite` | Ability to write devices |
-
-## Jobs
-
-| Name | Description |
-| ---- | ----------- |
-| `CreateJob` | Ability to start new jobs |
-| `UpdateJob` | Rename a job, modify its priority, etc... |
-| `DeleteJob` | Delete a job properties |
-| `ExecuteJob` | Allows updating a job metadata (name, changelist number, step properties, new groups, job states, etc...). Typically granted to agents. Not user facing. |
-| `RetryJobStep` | Ability to retry a failed job step |
-| `ViewJob` | Ability to view a job |
-
 ## Leases
 
 | Name | Description |
@@ -240,12 +198,6 @@
 | `UpdateNotice` |  |
 | `DeleteNotice` |  |
 
-## Notifications
-
-| Name | Description |
-| ---- | ----------- |
-| `CreateSubscription` | Ability to subscribe to notifications |
-
 ## Pools
 
 | Name | Description |
@@ -256,22 +208,6 @@
 | `ViewPool` |  |
 | `ListPools` |  |
 
-## Projects
-
-| Name | Description |
-| ---- | ----------- |
-| `CreateProject` | Allows the creation of new projects |
-| `DeleteProject` | Allows deletion of projects. |
-| `UpdateProject` | Modify attributes of a project (name, categories, etc...) |
-| `ViewProject` | View information about a project |
-
-## Replicators
-
-| Name | Description |
-| ---- | ----------- |
-| `UpdateReplicator` | Allows deletion of projects. |
-| `ViewReplicator` | Allows the creation of new projects |
-
 ## ServiceAccounts
 
 | Name | Description |
@@ -296,24 +232,6 @@
 | `DownloadSoftware` |  |
 | `DeleteSoftware` |  |
 
-## Streams
-
-| Name | Description |
-| ---- | ----------- |
-| `CreateStream` | Allows the creation of new streams within a project |
-| `UpdateStream` | Allows updating a stream (agent types, templates, schedules) |
-| `DeleteStream` | Allows deleting a stream |
-| `ViewStream` | Ability to view a stream |
-| `ViewChanges` | View changes submitted to a stream. NOTE: this returns responses from the server's Perforce account, which may be a priviledged user. |
-| `ViewTemplate` | View template associated with a stream |
-
-## TestData
-
-| Name | Description |
-| ---- | ----------- |
-| `TestHealthRead` | Ability to read test health |
-| `TestHealthWrite` | Ability to write test health |
-
 ## Accounts
 
 | Name | Description |
@@ -334,30 +252,6 @@
 | `ViewAgent` |  |
 | `ListAgents` |  |
 
-## Artifacts
-
-| Name | Description |
-| ---- | ----------- |
-| `ReadArtifact` |  |
-| `WriteArtifact` |  |
-| `DeleteArtifact` |  |
-| `UploadArtifact` |  |
-| `DownloadArtifact` |  |
-
-## Bisect
-
-| Name | Description |
-| ---- | ----------- |
-| `CreateBisectTask` |  |
-| `UpdateBisectTask` |  |
-| `ViewBisectTask` |  |
-
-## BuildHealth
-
-| Name | Description |
-| ---- | ----------- |
-| `AddUpdateDeleteBuildHealthFilter` |  |
-
 ## Compute
 
 | Name | Description |
@@ -367,24 +261,6 @@
 | `UbaCacheRead` |  |
 | `UbaCacheWrite` |  |
 
-## Devices
-
-| Name | Description |
-| ---- | ----------- |
-| `DeviceRead` |  |
-| `DeviceWrite` |  |
-
-## Jobs
-
-| Name | Description |
-| ---- | ----------- |
-| `CreateJob` |  |
-| `UpdateJob` |  |
-| `DeleteJob` |  |
-| `ExecuteJob` |  |
-| `RetryJobStep` |  |
-| `ViewJob` |  |
-
 ## Leases
 
 | Name | Description |
@@ -411,12 +287,6 @@
 | `UpdateNotice` |  |
 | `DeleteNotice` |  |
 
-## Notifications
-
-| Name | Description |
-| ---- | ----------- |
-| `CreateSubscription` |  |
-
 ## Pools
 
 | Name | Description |
@@ -427,22 +297,6 @@
 | `ViewPool` |  |
 | `ListPools` |  |
 
-## Projects
-
-| Name | Description |
-| ---- | ----------- |
-| `CreateProject` |  |
-| `DeleteProject` |  |
-| `UpdateProject` |  |
-| `ViewProject` |  |
-
-## Replicators
-
-| Name | Description |
-| ---- | ----------- |
-| `UpdateReplicator` |  |
-| `ViewReplicator` |  |
-
 ## ServiceAccounts
 
 | Name | Description |
@@ -467,24 +321,6 @@
 | `DownloadSoftware` |  |
 | `DeleteSoftware` |  |
 
-## Streams
-
-| Name | Description |
-| ---- | ----------- |
-| `CreateStream` |  |
-| `UpdateStream` |  |
-| `DeleteStream` |  |
-| `ViewStream` |  |
-| `ViewChanges` |  |
-| `ViewTemplate` |  |
-
-## TestData
-
-| Name | Description |
-| ---- | ----------- |
-| `TestHealthRead` |  |
-| `TestHealthWrite` |  |
-
 ## Accounts
 
 | Name | Description |
@@ -505,30 +341,6 @@
 | `ViewAgent` |  |
 | `ListAgents` |  |
 
-## Artifacts
-
-| Name | Description |
-| ---- | ----------- |
-| `ReadArtifact` |  |
-| `WriteArtifact` |  |
-| `DeleteArtifact` |  |
-| `UploadArtifact` |  |
-| `DownloadArtifact` |  |
-
-## Bisect
-
-| Name | Description |
-| ---- | ----------- |
-| `CreateBisectTask` |  |
-| `UpdateBisectTask` |  |
-| `ViewBisectTask` |  |
-
-## BuildHealth
-
-| Name | Description |
-| ---- | ----------- |
-| `AddUpdateDeleteBuildHealthFilter` |  |
-
 ## Compute
 
 | Name | Description |
@@ -538,24 +350,6 @@
 | `UbaCacheRead` |  |
 | `UbaCacheWrite` |  |
 
-## Devices
-
-| Name | Description |
-| ---- | ----------- |
-| `DeviceRead` |  |
-| `DeviceWrite` |  |
-
-## Jobs
-
-| Name | Description |
-| ---- | ----------- |
-| `CreateJob` |  |
-| `UpdateJob` |  |
-| `DeleteJob` |  |
-| `ExecuteJob` |  |
-| `RetryJobStep` |  |
-| `ViewJob` |  |
-
 ## Leases
 
 | Name | Description |
@@ -582,12 +376,6 @@
 | `UpdateNotice` |  |
 | `DeleteNotice` |  |
 
-## Notifications
-
-| Name | Description |
-| ---- | ----------- |
-| `CreateSubscription` |  |
-
 ## Pools
 
 | Name | Description |
@@ -598,22 +386,6 @@
 | `ViewPool` |  |
 | `ListPools` |  |
 
-## Projects
-
-| Name | Description |
-| ---- | ----------- |
-| `CreateProject` |  |
-| `DeleteProject` |  |
-| `UpdateProject` |  |
-| `ViewProject` |  |
-
-## Replicators
-
-| Name | Description |
-| ---- | ----------- |
-| `UpdateReplicator` |  |
-| `ViewReplicator` |  |
-
 ## Secrets
 
 | Name | Description |
@@ -644,24 +416,6 @@
 | `DownloadSoftware` |  |
 | `DeleteSoftware` |  |
 
-## Streams
-
-| Name | Description |
-| ---- | ----------- |
-| `CreateStream` |  |
-| `UpdateStream` |  |
-| `DeleteStream` |  |
-| `ViewStream` |  |
-| `ViewChanges` |  |
-| `ViewTemplate` |  |
-
-## TestData
-
-| Name | Description |
-| ---- | ----------- |
-| `TestHealthRead` |  |
-| `TestHealthWrite` |  |
-
 ## Accounts
 
 | Name | Description |
@@ -682,30 +436,6 @@
 | `ViewAgent` |  |
 | `ListAgents` |  |
 
-## Artifacts
-
-| Name | Description |
-| ---- | ----------- |
-| `ReadArtifact` |  |
-| `WriteArtifact` |  |
-| `DeleteArtifact` |  |
-| `UploadArtifact` |  |
-| `DownloadArtifact` |  |
-
-## Bisect
-
-| Name | Description |
-| ---- | ----------- |
-| `CreateBisectTask` |  |
-| `UpdateBisectTask` |  |
-| `ViewBisectTask` |  |
-
-## BuildHealth
-
-| Name | Description |
-| ---- | ----------- |
-| `AddUpdateDeleteBuildHealthFilter` |  |
-
 ## Compute
 
 | Name | Description |
@@ -715,24 +445,6 @@
 | `UbaCacheRead` |  |
 | `UbaCacheWrite` |  |
 
-## Devices
-
-| Name | Description |
-| ---- | ----------- |
-| `DeviceRead` |  |
-| `DeviceWrite` |  |
-
-## Jobs
-
-| Name | Description |
-| ---- | ----------- |
-| `CreateJob` |  |
-| `UpdateJob` |  |
-| `DeleteJob` |  |
-| `ExecuteJob` |  |
-| `RetryJobStep` |  |
-| `ViewJob` |  |
-
 ## Leases
 
 | Name | Description |
@@ -759,37 +471,15 @@
 | `UpdateNotice` |  |
 | `DeleteNotice` |  |
 
-## Notifications
-
-| Name | Description |
-| ---- | ----------- |
-| `CreateSubscription` |  |
-
 ## Pools
 
 | Name | Description |
 | ---- | ----------- |
-| `CreatePool` |  |
-| `UpdatePool` |  |
-| `DeletePool` |  |
-| `ViewPool` |  |
-| `ListPools` |  |
-
-## Projects
-
-| Name | Description |
-| ---- | ----------- |
-| `CreateProject` |  |
-| `DeleteProject` |  |
-| `UpdateProject` |  |
-| `ViewProject` |  |
-
-## Replicators
-
-| Name | Description |
-| ---- | ----------- |
-| `UpdateReplicator` |  |
-| `ViewReplicator` |  |
+| `CreatePool` |  |
+| `UpdatePool` |  |
+| `DeletePool` |  |
+| `ViewPool` |  |
+| `ListPools` |  |
 
 ## Secrets
 
@@ -834,24 +524,6 @@
 | `WriteRefs` | Ability to write refs to the storage service |
 | `DeleteRefs` | Ability to delete refs |
 
-## Streams
-
-| Name | Description |
-| ---- | ----------- |
-| `CreateStream` |  |
-| `UpdateStream` |  |
-| `DeleteStream` |  |
-| `ViewStream` |  |
-| `ViewChanges` |  |
-| `ViewTemplate` |  |
-
-## TestData
-
-| Name | Description |
-| ---- | ----------- |
-| `TestHealthRead` |  |
-| `TestHealthWrite` |  |
-
 ## Accounts
 
 | Name | Description |
@@ -872,30 +544,6 @@
 | `ViewAgent` |  |
 | `ListAgents` |  |
 
-## Artifacts
-
-| Name | Description |
-| ---- | ----------- |
-| `ReadArtifact` |  |
-| `WriteArtifact` |  |
-| `DeleteArtifact` |  |
-| `UploadArtifact` |  |
-| `DownloadArtifact` |  |
-
-## Bisect
-
-| Name | Description |
-| ---- | ----------- |
-| `CreateBisectTask` |  |
-| `UpdateBisectTask` |  |
-| `ViewBisectTask` |  |
-
-## BuildHealth
-
-| Name | Description |
-| ---- | ----------- |
-| `AddUpdateDeleteBuildHealthFilter` |  |
-
 ## Compute
 
 | Name | Description |
@@ -905,24 +553,6 @@
 | `UbaCacheRead` |  |
 | `UbaCacheWrite` |  |
 
-## Devices
-
-| Name | Description |
-| ---- | ----------- |
-| `DeviceRead` |  |
-| `DeviceWrite` |  |
-
-## Jobs
-
-| Name | Description |
-| ---- | ----------- |
-| `CreateJob` |  |
-| `UpdateJob` |  |
-| `DeleteJob` |  |
-| `ExecuteJob` |  |
-| `RetryJobStep` |  |
-| `ViewJob` |  |
-
 ## Leases
 
 | Name | Description |
@@ -949,12 +579,6 @@
 | `UpdateNotice` |  |
 | `DeleteNotice` |  |
 
-## Notifications
-
-| Name | Description |
-| ---- | ----------- |
-| `CreateSubscription` |  |
-
 ## Pools
 
 | Name | Description |
@@ -965,22 +589,6 @@
 | `ViewPool` |  |
 | `ListPools` |  |
 
-## Projects
-
-| Name | Description |
-| ---- | ----------- |
-| `CreateProject` |  |
-| `DeleteProject` |  |
-| `UpdateProject` |  |
-| `ViewProject` |  |
-
-## Replicators
-
-| Name | Description |
-| ---- | ----------- |
-| `UpdateReplicator` |  |
-| `ViewReplicator` |  |
-
 ## Secrets
 
 | Name | Description |
@@ -1024,23 +632,11 @@
 | `WriteRefs` |  |
 | `DeleteRefs` |  |
 
-## Streams
-
-| Name | Description |
-| ---- | ----------- |
-| `CreateStream` |  |
-| `UpdateStream` |  |
-| `DeleteStream` |  |
-| `ViewStream` |  |
-| `ViewChanges` |  |
-| `ViewTemplate` |  |
-
-## TestData
+## Symbols
 
 | Name | Description |
 | ---- | ----------- |
-| `TestHealthRead` |  |
-| `TestHealthWrite` |  |
+| `ReadSymbols` | Ability to download symbols |
 
 ## Accounts
 
@@ -1062,30 +658,6 @@
 | `ViewAgent` |  |
 | `ListAgents` |  |
 
-## Artifacts
-
-| Name | Description |
-| ---- | ----------- |
-| `ReadArtifact` |  |
-| `WriteArtifact` |  |
-| `DeleteArtifact` |  |
-| `UploadArtifact` |  |
-| `DownloadArtifact` |  |
-
-## Bisect
-
-| Name | Description |
-| ---- | ----------- |
-| `CreateBisectTask` |  |
-| `UpdateBisectTask` |  |
-| `ViewBisectTask` |  |
-
-## BuildHealth
-
-| Name | Description |
-| ---- | ----------- |
-| `AddUpdateDeleteBuildHealthFilter` |  |
-
 ## Compute
 
 | Name | Description |
@@ -1095,24 +667,6 @@
 | `UbaCacheRead` |  |
 | `UbaCacheWrite` |  |
 
-## Devices
-
-| Name | Description |
-| ---- | ----------- |
-| `DeviceRead` |  |
-| `DeviceWrite` |  |
-
-## Jobs
-
-| Name | Description |
-| ---- | ----------- |
-| `CreateJob` |  |
-| `UpdateJob` |  |
-| `DeleteJob` |  |
-| `ExecuteJob` |  |
-| `RetryJobStep` |  |
-| `ViewJob` |  |
-
 ## Leases
 
 | Name | Description |
@@ -1139,12 +693,6 @@
 | `UpdateNotice` |  |
 | `DeleteNotice` |  |
 
-## Notifications
-
-| Name | Description |
-| ---- | ----------- |
-| `CreateSubscription` |  |
-
 ## Pools
 
 | Name | Description |
@@ -1155,22 +703,6 @@
 | `ViewPool` |  |
 | `ListPools` |  |
 
-## Projects
-
-| Name | Description |
-| ---- | ----------- |
-| `CreateProject` |  |
-| `DeleteProject` |  |
-| `UpdateProject` |  |
-| `ViewProject` |  |
-
-## Replicators
-
-| Name | Description |
-| ---- | ----------- |
-| `UpdateReplicator` |  |
-| `ViewReplicator` |  |
-
 ## Secrets
 
 | Name | Description |
@@ -1214,29 +746,18 @@
 | `WriteRefs` |  |
 | `DeleteRefs` |  |
 
-## Streams
-
-| Name | Description |
-| ---- | ----------- |
-| `CreateStream` |  |
-| `UpdateStream` |  |
-| `DeleteStream` |  |
-| `ViewStream` |  |
-| `ViewChanges` |  |
-| `ViewTemplate` |  |
-
 ## Symbols
 
 | Name | Description |
 | ---- | ----------- |
-| `ReadSymbols` | Ability to download symbols |
+| `ReadSymbols` |  |
 
-## TestData
+## Tools
 
 | Name | Description |
 | ---- | ----------- |
-| `TestHealthRead` |  |
-| `TestHealthWrite` |  |
+| `DownloadTool` | Ability to download a tool |
+| `UploadTool` | Ability to upload new tool versions |
 
 ## Accounts
 
@@ -1262,25 +783,25 @@
 
 | Name | Description |
 | ---- | ----------- |
-| `ReadArtifact` |  |
-| `WriteArtifact` |  |
-| `DeleteArtifact` |  |
-| `UploadArtifact` |  |
-| `DownloadArtifact` |  |
+| `ReadArtifact` | Permission to read from an artifact |
+| `WriteArtifact` | Permission to write to an artifact |
+| `DeleteArtifact` | Permission to delete to an artifact |
+| `UploadArtifact` | Ability to create an artifact. Typically just for debugging; agents have this access for a particular session. |
+| `DownloadArtifact` | Ability to download an artifact |
 
 ## Bisect
 
 | Name | Description |
 | ---- | ----------- |
-| `CreateBisectTask` |  |
-| `UpdateBisectTask` |  |
-| `ViewBisectTask` |  |
+| `CreateBisectTask` | Ability to start new bisect tasks |
+| `UpdateBisectTask` | Ability to update a bisect task |
+| `ViewBisectTask` | Ability to view a bisect task |
 
 ## BuildHealth
 
 | Name | Description |
 | ---- | ----------- |
-| `AddUpdateDeleteBuildHealthFilter` |  |
+| `AddUpdateDeleteBuildHealthFilter` | Ability to add, update, and delete build health filters. |
 
 ## Compute
 
@@ -1295,19 +816,19 @@
 
 | Name | Description |
 | ---- | ----------- |
-| `DeviceRead` |  |
-| `DeviceWrite` |  |
+| `DeviceRead` | Ability to read devices |
+| `DeviceWrite` | Ability to write devices |
 
 ## Jobs
 
 | Name | Description |
 | ---- | ----------- |
-| `CreateJob` |  |
-| `UpdateJob` |  |
-| `DeleteJob` |  |
-| `ExecuteJob` |  |
-| `RetryJobStep` |  |
-| `ViewJob` |  |
+| `CreateJob` | Ability to start new jobs |
+| `UpdateJob` | Rename a job, modify its priority, etc... |
+| `DeleteJob` | Delete a job properties |
+| `ExecuteJob` | Allows updating a job metadata (name, changelist number, step properties, new groups, job states, etc...). Typically granted to agents. Not user facing. |
+| `RetryJobStep` | Ability to retry a failed job step |
+| `ViewJob` | Ability to view a job |
 
 ## Leases
 
@@ -1339,7 +860,7 @@
 
 | Name | Description |
 | ---- | ----------- |
-| `CreateSubscription` |  |
+| `CreateSubscription` | Ability to subscribe to notifications |
 
 ## Pools
 
@@ -1355,17 +876,17 @@
 
 | Name | Description |
 | ---- | ----------- |
-| `CreateProject` |  |
-| `DeleteProject` |  |
-| `UpdateProject` |  |
-| `ViewProject` |  |
+| `CreateProject` | Allows the creation of new projects |
+| `DeleteProject` | Allows deletion of projects. |
+| `UpdateProject` | Modify attributes of a project (name, categories, etc...) |
+| `ViewProject` | View information about a project |
 
 ## Replicators
 
 | Name | Description |
 | ---- | ----------- |
-| `UpdateReplicator` |  |
-| `ViewReplicator` |  |
+| `UpdateReplicator` | Allows deletion of projects. |
+| `ViewReplicator` | Allows the creation of new projects |
 
 ## Secrets
 
@@ -1414,12 +935,12 @@
 
 | Name | Description |
 | ---- | ----------- |
-| `CreateStream` |  |
-| `UpdateStream` |  |
-| `DeleteStream` |  |
-| `ViewStream` |  |
-| `ViewChanges` |  |
-| `ViewTemplate` |  |
+| `CreateStream` | Allows the creation of new streams within a project |
+| `UpdateStream` | Allows updating a stream (agent types, templates, schedules) |
+| `DeleteStream` | Allows deleting a stream |
+| `ViewStream` | Ability to view a stream |
+| `ViewChanges` | View changes submitted to a stream. NOTE: this returns responses from the server's Perforce account, which may be a priviledged user. |
+| `ViewTemplate` | View template associated with a stream |
 
 ## Symbols
 
@@ -1431,15 +952,15 @@
 
 | Name | Description |
 | ---- | ----------- |
-| `TestHealthRead` |  |
-| `TestHealthWrite` |  |
+| `TestHealthRead` | Ability to read test health |
+| `TestHealthWrite` | Ability to write test health |
 
 ## Tools
 
 | Name | Description |
 | ---- | ----------- |
-| `DownloadTool` | Ability to download a tool |
-| `UploadTool` | Ability to upload new tool versions |
+| `DownloadTool` |  |
+| `UploadTool` |  |
 
 ## Accounts
 
diff --git a/Engine/Source/Programs/Horde/Docs/Config/Schema/Globals.md b/Engine/Source/Programs/Horde/Docs/Config/Schema/Globals.md
--- a/Engine/Source/Programs/Horde/Docs/Config/Schema/Globals.md
+++ b/Engine/Source/Programs/Horde/Docs/Config/Schema/Globals.md
@@ -61,13 +61,13 @@
 ---- | -----------
 `analytics` | [AnalyticsConfig](#analyticsconfig)<br>Configuration for the analytics plugin
 `compute` | [ComputeConfig](#computeconfig)<br>Configuration for the compute plugin
-`build` | [BuildConfig](#buildconfig)<br>Configuration for the build plugin
 `exampleplugin` | [ExamplePluginConfig](#examplepluginconfig)<br>Configuration for the exampleplugin plugin
+`experimental` | [ExperimentalConfig](#experimentalconfig)<br>Configuration for the experimental plugin
 `secrets` | [SecretsConfig](#secretsconfig)<br>Configuration for the secrets plugin
 `storage` | [StorageConfig](#storageconfig)<br>Configuration for the storage plugin
-`experimental` | [ExperimentalConfig](#experimentalconfig)<br>Configuration for the experimental plugin
 `symbols` | [SymbolsConfig](#symbolsconfig)<br>Configuration for the symbols plugin
 `tools` | [ToolsConfig](#toolsconfig)<br>Configuration for the tools plugin
+`build` | [BuildConfig](#buildconfig)<br>Configuration for the build plugin
 
 ## AnalyticsConfig
 
@@ -317,153 +317,6 @@
 `description` | `string`<br>Human-readable description
 `computeId` | `string`<br>Compute ID for this network (used when allocating compute resources)
 
-## BuildConfig
-
-Configuration for the build plugin
-
-Name | Description
----- | -----------
-`perforceClusters` | [PerforceCluster](#perforcecluster)`[]`<br>List of Perforce clusters
-`devices` | [DeviceConfig](#deviceconfig)<br>Device configuration
-`maxConformCount` | `integer`<br>Maximum number of conforms to run at once
-`agentShutdownIfDisabledGracePeriod` | `string`<br>Time to wait before shutting down an agent that has been disabled Used if no value is set on the actual pool.
-`robomerge` | [RoboMergeConfig](#robomergeconfig)<br>RoboMerge branchmap file configuration for the merge graph service. When null or empty, the service is inert (no Perforce connections, no API data).
-`artifactTypes` | [ArtifactTypeConfig](#artifacttypeconfig)`[]`<br>Configuration for different artifact types
-`projects` | [ProjectConfig](Projects.md)`[]`<br>List of projects
-`enableConformTasks` | `boolean`<br>Whether to allow conform tasks to run
-`enableJobTaskConformCheck` | `boolean`<br>Whether to check if a conform is required on an agent before assigning a lease for a job task source.
-`issueFixedTag` | `string`<br>Commit tag to use for marking issues as fixed
-`issues` | [IssueConfig](Issues.md)<br>Build health issue configurations
-`enableIssueRootCauseAnalysisNotifications` | `boolean`<br>Whether to allow notifications for Root Cause Analysis to be sent
-`enableAutoRcaSlackNotifications` | `boolean`<br>Global kill switch for posting auto-RCA results to Slack triage threads. When false (default), no auto-RCA Slack notifications are sent regardless of per-workflow AutoRcaSlackEnabled settings.
-`autoRcaUseMultiStage` | `boolean`<br>Whether to use multi-stage LLM analysis for auto-RCA. Multi-stage is slower but potentially more accurate. Set to false for faster processing with higher throughput.
-`maxConcurrentRcaRequests` | `integer`<br>Maximum number of concurrent auto-RCA requests. Limits load on the external LLM API. Set to 0 to disable auto-RCA globally (acts as a kill switch).
-`issueFindSuspectsMaxSuccessDays` | `integer`<br>The maximum days since node success to scan for suspects
-`issueFindSuspectsMaxScoringSeconds` | `integer`<br>The maximum number of seconds to spend scording suspects
-`issueFindSuspectsMaxChanges` | `integer`<br>Maximum number of commits to examine for suspects. Previously hardcoded as const MaxChanges = 500.
-`issueFindSuspectsEnableParallelScoring` | `boolean`<br>Enable parallel scoring of commits for suspect finding. When enabled, commits are scored concurrently up to MaxParallelism.
-`issueFindSuspectsMaxParallelism` | `integer`<br>Maximum concurrent commit file lookups when parallel scoring is enabled. Higher values improve speed but increase Perforce server load.
-`issueFindSuspectsEnableEarlyExit` | `boolean`<br>Enable early exit when theoretical maximum score is achieved. Only exits early if the best possible score has been found.
-`issueFindSuspectsEnableBidirectionalQuery` | `boolean`<br>Enable bidirectional commit query to capture both recent and old commits. Queries half the limit from each end of the commit range.
-`issueFindSuspectsEnableFingerprintFallback` | `boolean`<br>Enable fingerprint-based fallback when no suspects are found. Helps prevent duplicate issues when commit range exceeds MaxChanges.
-`issueFingerprintFallbackMinScore` | `integer`<br>Minimum score required for fingerprint fallback matching. Higher values are more conservative (fewer false matches).
-`issueFingerprintFallbackMaxAgeDays` | `integer`<br>Maximum age in days for issues to be considered in fingerprint fallback.
-`issueFingerprintFallbackExcludedTypes` | `string[]`<br>Fingerprint types excluded from fallback matching (too generic).
-`enableExpandedIgnorePatterns` | `boolean`<br>Whether to enable expanded IgnorePatterns discovery (workspace-wide scanning + config-specified files). When false (default), only Engine directory is scanned for ignore patterns.
-`lowContextFingerprintTypes` | `string[]`<br>Fingerprint type prefixes considered "low context" for suspect mention suppression. When an issue's fingerprint type starts with any of these prefixes AND all suspects are tied at rank 0, @ mentions are suppressed in Slack. Suspects are still listed without mentions. Set to an empty array to disable suppression entirely.
-
-## PerforceCluster
-
-Information about a cluster of Perforce servers.
-
-Name | Description
----- | -----------
-`name` | `string`<br>Name of the cluster
-`serviceAccount` | `string`<br>Username for Horde to log in to this server. Will use the first account specified below if not overridden.
-`canImpersonate` | `boolean`<br>Whether the service account can impersonate other users
-`supportsPartitionedWorkspaces` | `boolean`<br>Whether to use partitioned workspaces on this server
-`servers` | [PerforceServer](#perforceserver)`[]`<br>List of servers
-`credentials` | [PerforceCredentials](#perforcecredentials)`[]`<br>List of server credentials
-`autoSdk` | [AutoSdkWorkspace](#autosdkworkspace)`[]`<br>List of autosdk streams
-
-## PerforceServer
-
-Information about an individual Perforce server
-
-Name | Description
----- | -----------
-`serverAndPort` | `string`<br>The server and port. The server may be a DNS entry with multiple records, in which case it will be actively load balanced. If "ssl:" prefix is used, ensure P4 server's fingerprint/certificate is trusted. See Horde's documentation on connecting to SSL-enabled Perforce servers.
-`healthCheck` | `boolean`<br>Whether to query the healthcheck address under each server
-`resolveDns` | `boolean`<br>Whether to resolve the DNS entries and load balance between different hosts
-`maxConformCount` | `integer`<br>Maximum number of simultaneous conforms on this server
-`condition` | `string`<br>Optional condition for a machine to be eligible to use this server
-`properties` | `string[]`<br>List of properties for an agent to be eligible to use this server
-
-## PerforceCredentials
-
-Credentials for a Perforce user
-
-Name | Description
----- | -----------
-`userName` | `string`<br>The username
-`password` | `string`<br>Password for the user
-`ticket` | `string`<br>Login ticket for the user (will be used instead of password if set)
-
-## AutoSdkWorkspace
-
-Path to a platform and stream to use for syncing AutoSDK
-
-Name | Description
----- | -----------
-`name` | `string`<br>Name of this workspace
-`properties` | `string[]`<br>The agent properties to check (eg. "OSFamily=Windows")
-`userName` | `string`<br>Username for logging in to the server
-`stream` | `string`<br>Stream to use
-
-## DeviceConfig
-
-Configuration for devices
-
-Name | Description
----- | -----------
-`platforms` | [DevicePlatformConfig](#deviceplatformconfig)`[]`<br>List of device platforms
-`pools` | [DevicePoolConfig](#devicepoolconfig)`[]`<br>List of device pools
-
-## DevicePlatformConfig
-
-Configuration for a device platform
-
-Name | Description
----- | -----------
-`id` | `string`<br>The id for this platform
-`name` | `string`<br>Name of the platform
-`models` | `string[]`<br>A list of platform models
-`legacyNames` | `string[]`<br>Legacy names which older versions of Gauntlet may be using
-`legacyPerfSpecHighModel` | `string`<br>Model name for the high perf spec, which may be requested by Gauntlet
-
-## DevicePoolConfig
-
-Configuration for a device pool
-
-Name | Description
----- | -----------
-`id` | `string`<br>The id for this platform
-`name` | `string`<br>The name of the pool
-`poolType` | [DevicePoolType](#devicepooltype-enum)<br>The type of the pool
-`projectIds` | `string[]`<br>List of project ids associated with pool
-
-## DevicePoolType (Enum)
-
-The type of device pool
-
-Name | Description
----- | -----------
-`Automation` | Available to CIS jobs
-`Shared` | Shared by users with remote checking and checkouts
-
-## RoboMergeConfig
-
-Configuration for the RoboMerge branchmap service
-
-Name | Description
----- | -----------
-`branchmapFiles` | `string[]`<br>Perforce depot paths to branchmap JSON files. Plain depot paths (e.g. "//MyProject/Main/Build/Streams/myproject.branchmap.json"). All files are read from the Perforce cluster specified by .
-`perforceCluster` | `string`<br>Perforce cluster name to connect to for reading branchmap files. Maps to a connection configured in BuildServerConfig.Perforce[]. Default: "default" (uses the default Perforce connection).
-`tickInterval` | `string`<br>How often to poll Perforce for branchmap file changes. Default: 5 minutes. Changing this requires a server restart (ticker interval is set at construction time).
-
-## ArtifactTypeConfig
-
-Configuration for an artifact
-
-Name | Description
----- | -----------
-`name` | `string`<br>Legacy 'Name' property
-`type` | `string`<br>Name of the artifact type
-`acl` | [AclConfig](#aclconfig)<br>Acl for the artifact type
-`keepCount` | `integer`<br>Number of artifacts to retain
-`keepDays` | `integer`<br>Number of days to retain artifacts of this type
-`namespaceId` | `string`<br>Storage namespace to use for this artifact types
-
 ## ExamplePluginConfig
 
 Configuration for the ExamplePlugin plugin
@@ -472,5 +325,13 @@
 ---- | -----------
 `greetingMessage` | `string`<br>The Greeting provided by the example plugin
 
+## ExperimentalConfig
+
+Configuration for the experimental plugin
+
+Name | Description
+---- | -----------
+`notifications` | [NotificationConfig](Notification.md)`[]`<br>Notifications to be sent on the Horde server
+
 ## SecretsConfig
 
@@ -632,14 +493,6 @@
 `enableAliases` | `boolean`<br>Support querying exports by their aliases
 `acl` | [AclConfig](#aclconfig)<br>Access list for this namespace
 
-## ExperimentalConfig
-
-Configuration for the experimental plugin
-
-Name | Description
----- | -----------
-`notifications` | [NotificationConfig](Notification.md)`[]`<br>Notifications to be sent on the Horde server
-
 ## SymbolsConfig
 
 Configuration for the tools system
@@ -687,3 +540,150 @@
 `namespaceId` | `string`<br>Default namespace for new deployments of this tool
 `maxDeployments` | `integer`<br>Maximum number of deployments to retain for this tool. When exceeded, the oldest deployments are removed.
 `acl` | [AclConfig](#aclconfig)<br>Permissions for the tool
+
+## BuildConfig
+
+Configuration for the build plugin
+
+Name | Description
+---- | -----------
+`perforceClusters` | [PerforceCluster](#perforcecluster)`[]`<br>List of Perforce clusters
+`devices` | [DeviceConfig](#deviceconfig)<br>Device configuration
+`maxConformCount` | `integer`<br>Maximum number of conforms to run at once
+`agentShutdownIfDisabledGracePeriod` | `string`<br>Time to wait before shutting down an agent that has been disabled Used if no value is set on the actual pool.
+`robomerge` | [RoboMergeConfig](#robomergeconfig)<br>RoboMerge branchmap file configuration for the merge graph service. When null or empty, the service is inert (no Perforce connections, no API data).
+`artifactTypes` | [ArtifactTypeConfig](#artifacttypeconfig)`[]`<br>Configuration for different artifact types
+`projects` | [ProjectConfig](Projects.md)`[]`<br>List of projects
+`enableConformTasks` | `boolean`<br>Whether to allow conform tasks to run
+`enableJobTaskConformCheck` | `boolean`<br>Whether to check if a conform is required on an agent before assigning a lease for a job task source.
+`issueFixedTag` | `string`<br>Commit tag to use for marking issues as fixed
+`issues` | [IssueConfig](Issues.md)<br>Build health issue configurations
+`enableIssueRootCauseAnalysisNotifications` | `boolean`<br>Whether to allow notifications for Root Cause Analysis to be sent
+`enableAutoRcaSlackNotifications` | `boolean`<br>Global kill switch for posting auto-RCA results to Slack triage threads. When false (default), no auto-RCA Slack notifications are sent regardless of per-workflow AutoRcaSlackEnabled settings.
+`autoRcaUseMultiStage` | `boolean`<br>Whether to use multi-stage LLM analysis for auto-RCA. Multi-stage is slower but potentially more accurate. Set to false for faster processing with higher throughput.
+`maxConcurrentRcaRequests` | `integer`<br>Maximum number of concurrent auto-RCA requests. Limits load on the external LLM API. Set to 0 to disable auto-RCA globally (acts as a kill switch).
+`issueFindSuspectsMaxSuccessDays` | `integer`<br>The maximum days since node success to scan for suspects
+`issueFindSuspectsMaxScoringSeconds` | `integer`<br>The maximum number of seconds to spend scording suspects
+`issueFindSuspectsMaxChanges` | `integer`<br>Maximum number of commits to examine for suspects. Previously hardcoded as const MaxChanges = 500.
+`issueFindSuspectsEnableParallelScoring` | `boolean`<br>Enable parallel scoring of commits for suspect finding. When enabled, commits are scored concurrently up to MaxParallelism.
+`issueFindSuspectsMaxParallelism` | `integer`<br>Maximum concurrent commit file lookups when parallel scoring is enabled. Higher values improve speed but increase Perforce server load.
+`issueFindSuspectsEnableEarlyExit` | `boolean`<br>Enable early exit when theoretical maximum score is achieved. Only exits early if the best possible score has been found.
+`issueFindSuspectsEnableBidirectionalQuery` | `boolean`<br>Enable bidirectional commit query to capture both recent and old commits. Queries half the limit from each end of the commit range.
+`issueFindSuspectsEnableFingerprintFallback` | `boolean`<br>Enable fingerprint-based fallback when no suspects are found. Helps prevent duplicate issues when commit range exceeds MaxChanges.
+`issueFingerprintFallbackMinScore` | `integer`<br>Minimum score required for fingerprint fallback matching. Higher values are more conservative (fewer false matches).
+`issueFingerprintFallbackMaxAgeDays` | `integer`<br>Maximum age in days for issues to be considered in fingerprint fallback.
+`issueFingerprintFallbackExcludedTypes` | `string[]`<br>Fingerprint types excluded from fallback matching (too generic).
+`enableExpandedIgnorePatterns` | `boolean`<br>Whether to enable expanded IgnorePatterns discovery (workspace-wide scanning + config-specified files). When false (default), only Engine directory is scanned for ignore patterns.
+`lowContextFingerprintTypes` | `string[]`<br>Fingerprint type prefixes considered "low context" for suspect mention suppression. When an issue's fingerprint type starts with any of these prefixes AND all suspects are tied at rank 0, @ mentions are suppressed in Slack. Suspects are still listed without mentions. Set to an empty array to disable suppression entirely.
+
+## PerforceCluster
+
+Information about a cluster of Perforce servers.
+
+Name | Description
+---- | -----------
+`name` | `string`<br>Name of the cluster
+`serviceAccount` | `string`<br>Username for Horde to log in to this server. Will use the first account specified below if not overridden.
+`canImpersonate` | `boolean`<br>Whether the service account can impersonate other users
+`supportsPartitionedWorkspaces` | `boolean`<br>Whether to use partitioned workspaces on this server
+`servers` | [PerforceServer](#perforceserver)`[]`<br>List of servers
+`credentials` | [PerforceCredentials](#perforcecredentials)`[]`<br>List of server credentials
+`autoSdk` | [AutoSdkWorkspace](#autosdkworkspace)`[]`<br>List of autosdk streams
+
+## PerforceServer
+
+Information about an individual Perforce server
+
+Name | Description
+---- | -----------
+`serverAndPort` | `string`<br>The server and port. The server may be a DNS entry with multiple records, in which case it will be actively load balanced. If "ssl:" prefix is used, ensure P4 server's fingerprint/certificate is trusted. See Horde's documentation on connecting to SSL-enabled Perforce servers.
+`healthCheck` | `boolean`<br>Whether to query the healthcheck address under each server
+`resolveDns` | `boolean`<br>Whether to resolve the DNS entries and load balance between different hosts
+`maxConformCount` | `integer`<br>Maximum number of simultaneous conforms on this server
+`condition` | `string`<br>Optional condition for a machine to be eligible to use this server
+`properties` | `string[]`<br>List of properties for an agent to be eligible to use this server
+
+## PerforceCredentials
+
+Credentials for a Perforce user
+
+Name | Description
+---- | -----------
+`userName` | `string`<br>The username
+`password` | `string`<br>Password for the user
+`ticket` | `string`<br>Login ticket for the user (will be used instead of password if set)
+
+## AutoSdkWorkspace
+
+Path to a platform and stream to use for syncing AutoSDK
+
+Name | Description
+---- | -----------
+`name` | `string`<br>Name of this workspace
+`properties` | `string[]`<br>The agent properties to check (eg. "OSFamily=Windows")
+`userName` | `string`<br>Username for logging in to the server
+`stream` | `string`<br>Stream to use
+
+## DeviceConfig
+
+Configuration for devices
+
+Name | Description
+---- | -----------
+`platforms` | [DevicePlatformConfig](#deviceplatformconfig)`[]`<br>List of device platforms
+`pools` | [DevicePoolConfig](#devicepoolconfig)`[]`<br>List of device pools
+
+## DevicePlatformConfig
+
+Configuration for a device platform
+
+Name | Description
+---- | -----------
+`id` | `string`<br>The id for this platform
+`name` | `string`<br>Name of the platform
+`models` | `string[]`<br>A list of platform models
+`legacyNames` | `string[]`<br>Legacy names which older versions of Gauntlet may be using
+`legacyPerfSpecHighModel` | `string`<br>Model name for the high perf spec, which may be requested by Gauntlet
+
+## DevicePoolConfig
+
+Configuration for a device pool
+
+Name | Description
+---- | -----------
+`id` | `string`<br>The id for this platform
+`name` | `string`<br>The name of the pool
+`poolType` | [DevicePoolType](#devicepooltype-enum)<br>The type of the pool
+`projectIds` | `string[]`<br>List of project ids associated with pool
+
+## DevicePoolType (Enum)
+
+The type of device pool
+
+Name | Description
+---- | -----------
+`Automation` | Available to CIS jobs
+`Shared` | Shared by users with remote checking and checkouts
+
+## RoboMergeConfig
+
+Configuration for the RoboMerge branchmap service
+
+Name | Description
+---- | -----------
+`branchmapFiles` | `string[]`<br>Perforce depot paths to branchmap JSON files. Plain depot paths (e.g. "//MyProject/Main/Build/Streams/myproject.branchmap.json"). All files are read from the Perforce cluster specified by .
+`perforceCluster` | `string`<br>Perforce cluster name to connect to for reading branchmap files. Maps to a connection configured in BuildServerConfig.Perforce[]. Default: "default" (uses the default Perforce connection).
+`tickInterval` | `string`<br>How often to poll Perforce for branchmap file changes. Default: 5 minutes. Changing this requires a server restart (ticker interval is set at construction time).
+
+## ArtifactTypeConfig
+
+Configuration for an artifact
+
+Name | Description
+---- | -----------
+`name` | `string`<br>Legacy 'Name' property
+`type` | `string`<br>Name of the artifact type
+`acl` | [AclConfig](#aclconfig)<br>Acl for the artifact type
+`keepCount` | `integer`<br>Number of artifacts to retain
+`keepDays` | `integer`<br>Number of days to retain artifacts of this type
+`namespaceId` | `string`<br>Storage namespace to use for this artifact types
diff --git a/Engine/Source/Programs/Horde/Docs/Deployment/ServerSettings.md b/Engine/Source/Programs/Horde/Docs/Deployment/ServerSettings.md
--- a/Engine/Source/Programs/Horde/Docs/Deployment/ServerSettings.md
+++ b/Engine/Source/Programs/Horde/Docs/Deployment/ServerSettings.md
@@ -129,13 +129,13 @@
 ---- | -----------
 `analytics` | [AnalyticsServerConfig](#analyticsserverconfig)<br>Configuration for the analytics plugin
 `compute` | [ComputeServerConfig](#computeserverconfig)<br>Configuration for the compute plugin
-`build` | [BuildServerConfig](#buildserverconfig)<br>Configuration for the build plugin
 `exampleplugin` | [ExamplePluginServerConfig](#examplepluginserverconfig)<br>Configuration for the exampleplugin plugin
+`experimental` | [ExperimentalServerConfig](#experimentalserverconfig)<br>Configuration for the experimental plugin
 `secrets` | [SecretsServerConfig](#secretsserverconfig)<br>Configuration for the secrets plugin
 `storage` | [StorageServerConfig](#storageserverconfig)<br>Configuration for the storage plugin
-`experimental` | [ExperimentalServerConfig](#experimentalserverconfig)<br>Configuration for the experimental plugin
 `symbols` | [PluginServerConfig](#pluginserverconfig)<br>Configuration for the symbols plugin
 `tools` | [ToolsServerConfig](#toolsserverconfig)<br>Configuration for the tools plugin
+`build` | [BuildServerConfig](#buildserverconfig)<br>Configuration for the build plugin
 
 ## AnalyticsServerConfig
 
@@ -219,154 +219,17 @@
 `ComputeQueueAwsMetric` | A no-op strategy that reports metrics to let an external AWS auto-scaling policy scale the fleet
 `LeaseUtilizationAwsMetric` | A no-op strategy that reports metrics to let an external AWS auto-scaling policy scale the fleet
 
-## BuildServerConfig
-
-Static configuration for the build plugin
-
-Name | Description
----- | -----------
-`perforce` | [PerforceConnectionSettings](#perforceconnectionsettings)`[]`<br>Perforce connections for use by the Horde server (not agents)
-`useLocalPerforceEnv` | `boolean`<br>Whether to use the local Perforce environment
-`perforceConnectionPoolSize` | `integer`<br>Number of pooled perforce connections to keep
-`enableConformTasks` | `boolean`<br>Whether to enable the conform task source.
-`p4SwarmUrl` | `string`<br>Url of P4 Swarm installation
-`robomergeUrl` | `string`<br>Url of Robomerge installation
-`commitsViewerUrl` | `string`<br>Url of Commits Viewer
-`jiraUsername` | `string`<br>The Jira service account user name
-`jiraApiToken` | `string`<br>The Jira service account API token
-`jiraUrl` | `string`<br>The Uri for the Jira installation
-`sharedDeviceCheckoutDays` | `integer`<br>The number of days shared device checkouts are held
-`deviceProblemCooldownMinutes` | `integer`<br>The number of cooldown minutes for device problems
-`deviceReportChannel` | `string`<br>Channel to send device reports to
-`disableBisection` | `boolean`<br>Whether to disable bisection tasks
-`disableSchedules` | `boolean`<br>Whether to run scheduled jobs.
-`slackToken` | `string`<br>Bot token for interacting with Slack (xoxb-*)
-`slackSocketToken` | `string`<br>Token for opening a socket to slack (xapp-*)
-`slackAdminToken` | `string`<br>Admin user token for Slack (xoxp-*). This is only required when using the admin endpoints to invite users.
-`slackUsers` | `string`<br>Filtered list of slack users to send notifications to. Should be Slack user ids, separated by commas.
-`slackErrorPrefix` | `string`<br>Prefix to use when reporting errors
-`slackWarningPrefix` | `string`<br>Prefix to use when reporting warnings
-`configNotificationChannel` | `string`<br>Channel for sending messages related to config update failures
-`updateStreamsNotificationChannel` | `string`<br>Channel to send stream notification update failures to
-`jobNotificationChannel` | `string`<br>Slack channel to send job related notifications to. Multiple channels can be specified, separated by ;
-`agentNotificationChannel` | `string`<br>Slack channel to send agent related notifications to.
-`testDataRetainMonths` | `integer`<br>The number of months to retain test data
-`testHealthHistoryEvaluationDays` | `integer`<br>The number of days to go back in historical data to evaluate test health
-`testHealthReportRemindDays` | `integer`<br>The number of days after which a test health report remind for improvement
-`blockCacheDir` | `string`<br>Directory to store the fine-grained block cache. This caches individual exports embedded in bundles.
-`blockCacheSize` | `string`<br>Maximum size of the block cache. Accepts standard binary suffixes. Currently only allocates in multiples of 1024mb.
-`blockCacheSizeBytes` | `integer`<br>Accessor for the block cache size in bytes
-`commits` | [CommitSettings](#commitsettings)<br>Options for the commit service
-`enabled` | `boolean`<br>Whether the plugin should be enabled or not
-
-## PerforceConnectionSettings
-
-Perforce connection information for use by the Horde server (for reading config files, etc...)
-
-Name | Description
----- | -----------
-`id` | `string`<br>Identifier for this server
-`serverAndPort` | `string`<br>Server and port
-`credentials` | [PerforceCredentials](#perforcecredentials)<br>Credentials for the server
-
-## PerforceCredentials
-
-Credentials for a Perforce user
-
-Name | Description
----- | -----------
-`userName` | `string`<br>The username
-`password` | `string`<br>Password for the user
-`ticket` | `string`<br>Login ticket for the user (will be used instead of password if set)
-
-## CommitSettings
-
-Options for the commit service
-
-Name | Description
----- | -----------
-`replicateMetadata` | `boolean`<br>Whether to mirror commit metadata to the database
-`replicateContent` | `boolean`<br>Whether to mirror commit data to storage
-`bundle` | [BundleOptions](#bundleoptions)<br>Options for how objects are packed together
-`chunking` | [ChunkingOptions](#chunkingoptions)<br>Options for how objects are sliced
-
-## BundleOptions
-
-Options for configuring a bundle serializer
-
-Name | Description
----- | -----------
-`maxVersion` | [BundleVersion](#bundleversion-enum)<br>Maximum version number of bundles to write
-`maxBlobSize` | `integer`<br>Maximum payload size fo a blob
-`compressionFormat` | [BundleCompressionFormat](#bundlecompressionformat-enum)<br>Compression format to use
-`minCompressionPacketSize` | `integer`<br>Minimum size of a block to be compressed
-`maxWriteQueueLength` | `integer`<br>Maximum amount of data to store in memory. This includes any background writes as well as bundles being built.
-
-## BundleVersion (Enum)
-
-Bundle version number
-
-Name | Description
----- | -----------
-`Initial` | Initial version number
-`ExportAliases` | Added the BundleExport.Alias property
-`RemoveAliases` | Back out change to include aliases. Will likely do this through an API rather than baked into the data.
-`InPlace` | Use data structures which support in-place reading and writing.
-`ImportHashes` | Add import hashes to imported nodes
-`LatestV1` | Last version using the V1 pipeline
-`PacketSequence` | Structure bundles as a sequence of self-contained packets (uses V2 code)
-`Latest` | The current version number
-`LatestV2` | Last version using the V2 pipeline
-`LatestPlusOne` | Last item in the enum. Used for
-
-## BundleCompressionFormat (Enum)
-
-Indicates the compression format in the bundle
-
-Name | Description
----- | -----------
-`None` | Packets are uncompressed
-`LZ4` | LZ4 compression
-`Gzip` | Gzip compression
-`Oodle` | Oodle compression (Selkie)
-`Brotli` | Brotli compression
-`Zstd` | ZStandard compression
-
-## ChunkingOptions
-
-Options for creating file nodes
-
-Name | Description
----- | -----------
-`leafOptions` | [LeafChunkedDataNodeOptions](#leafchunkeddatanodeoptions)<br>Options for creating leaf nodes
-`interiorOptions` | [InteriorChunkedDataNodeOptions](#interiorchunkeddatanodeoptions)<br>Options for creating interior nodes
-
-## LeafChunkedDataNodeOptions
-
-Options for creating a specific type of file nodes
-
-Name | Description
----- | -----------
-`minSize` | `integer`<br>Minimum chunk size
-`maxSize` | `integer`<br>Maximum chunk size. Chunks will be split on this boundary if another match is not found.
-`targetSize` | `integer`<br>Target chunk size for content-slicing
-`windowSize` | `integer`<br>Window size to use when scanning for split points
-`threshold` | `integer`<br>Accessor for the BuzHash chunking threshold
-
-## InteriorChunkedDataNodeOptions
-
-Options for creating interior nodes
-
-Name | Description
----- | -----------
-`minChildCount` | `integer`<br>Minimum number of children in each node
-`targetChildCount` | `integer`<br>Target number of children in each node
-`maxChildCount` | `integer`<br>Maximum number of children in each node
-`sliceThreshold` | `integer`<br>Threshold hash value for splitting interior nodes
-
 ## ExamplePluginServerConfig
 
 Server configuration for the ExamplePlugin
 
 Name | Description
 ---- | -----------
+`enabled` | `boolean`<br>Whether the plugin should be enabled or not
+
+## ExperimentalServerConfig
+
+Server configuration for the experimental plugin
+
+Name | Description
+---- | -----------
@@ -440,14 +303,6 @@
 `AssumeRole` | Assume a particular role. Should specify ARN in
 `AssumeRoleWebIdentity` | Assume a particular role using the current environment variables.
 
-## ExperimentalServerConfig
-
-Server configuration for the experimental plugin
-
-Name | Description
----- | -----------
-`enabled` | `boolean`<br>Whether the plugin should be enabled or not
-
 ## PluginServerConfig
 
 Base class for plugin server config objects
@@ -529,3 +384,148 @@
 `actions` | `string[]`<br>Actions to include
 `excludeActions` | `string[]`<br>Actions to exclude from the inherited actions
 `extends` | `string[]`<br>Other profiles to extend from
+
+## BuildServerConfig
+
+Static configuration for the build plugin
+
+Name | Description
+---- | -----------
+`perforce` | [PerforceConnectionSettings](#perforceconnectionsettings)`[]`<br>Perforce connections for use by the Horde server (not agents)
+`useLocalPerforceEnv` | `boolean`<br>Whether to use the local Perforce environment
+`perforceConnectionPoolSize` | `integer`<br>Number of pooled perforce connections to keep
+`enableConformTasks` | `boolean`<br>Whether to enable the conform task source.
+`p4SwarmUrl` | `string`<br>Url of P4 Swarm installation
+`robomergeUrl` | `string`<br>Url of Robomerge installation
+`commitsViewerUrl` | `string`<br>Url of Commits Viewer
+`jiraUsername` | `string`<br>The Jira service account user name
+`jiraApiToken` | `string`<br>The Jira service account API token
+`jiraUrl` | `string`<br>The Uri for the Jira installation
+`sharedDeviceCheckoutDays` | `integer`<br>The number of days shared device checkouts are held
+`deviceProblemCooldownMinutes` | `integer`<br>The number of cooldown minutes for device problems
+`deviceReportChannel` | `string`<br>Channel to send device reports to
+`disableBisection` | `boolean`<br>Whether to disable bisection tasks
+`disableSchedules` | `boolean`<br>Whether to run scheduled jobs.
+`slackToken` | `string`<br>Bot token for interacting with Slack (xoxb-*)
+`slackSocketToken` | `string`<br>Token for opening a socket to slack (xapp-*)
+`slackAdminToken` | `string`<br>Admin user token for Slack (xoxp-*). This is only required when using the admin endpoints to invite users.
+`slackUsers` | `string`<br>Filtered list of slack users to send notifications to. Should be Slack user ids, separated by commas.
+`slackErrorPrefix` | `string`<br>Prefix to use when reporting errors
+`slackWarningPrefix` | `string`<br>Prefix to use when reporting warnings
+`configNotificationChannel` | `string`<br>Channel for sending messages related to config update failures
+`updateStreamsNotificationChannel` | `string`<br>Channel to send stream notification update failures to
+`jobNotificationChannel` | `string`<br>Slack channel to send job related notifications to. Multiple channels can be specified, separated by ;
+`agentNotificationChannel` | `string`<br>Slack channel to send agent related notifications to.
+`testDataRetainMonths` | `integer`<br>The number of months to retain test data
+`testHealthHistoryEvaluationDays` | `integer`<br>The number of days to go back in historical data to evaluate test health
+`testHealthReportRemindDays` | `integer`<br>The number of days after which a test health report remind for improvement
+`blockCacheDir` | `string`<br>Directory to store the fine-grained block cache. This caches individual exports embedded in bundles.
+`blockCacheSize` | `string`<br>Maximum size of the block cache. Accepts standard binary suffixes. Currently only allocates in multiples of 1024mb.
+`blockCacheSizeBytes` | `integer`<br>Accessor for the block cache size in bytes
+`commits` | [CommitSettings](#commitsettings)<br>Options for the commit service
+`enabled` | `boolean`<br>Whether the plugin should be enabled or not
+
+## PerforceConnectionSettings
+
+Perforce connection information for use by the Horde server (for reading config files, etc...)
+
+Name | Description
+---- | -----------
+`id` | `string`<br>Identifier for this server
+`serverAndPort` | `string`<br>Server and port
+`credentials` | [PerforceCredentials](#perforcecredentials)<br>Credentials for the server
+
+## PerforceCredentials
+
+Credentials for a Perforce user
+
+Name | Description
+---- | -----------
+`userName` | `string`<br>The username
+`password` | `string`<br>Password for the user
+`ticket` | `string`<br>Login ticket for the user (will be used instead of password if set)
+
+## CommitSettings
+
+Options for the commit service
+
+Name | Description
+---- | -----------
+`replicateMetadata` | `boolean`<br>Whether to mirror commit metadata to the database
+`replicateContent` | `boolean`<br>Whether to mirror commit data to storage
+`bundle` | [BundleOptions](#bundleoptions)<br>Options for how objects are packed together
+`chunking` | [ChunkingOptions](#chunkingoptions)<br>Options for how objects are sliced
+
+## BundleOptions
+
+Options for configuring a bundle serializer
+
+Name | Description
+---- | -----------
+`maxVersion` | [BundleVersion](#bundleversion-enum)<br>Maximum version number of bundles to write
+`maxBlobSize` | `integer`<br>Maximum payload size fo a blob
+`compressionFormat` | [BundleCompressionFormat](#bundlecompressionformat-enum)<br>Compression format to use
+`minCompressionPacketSize` | `integer`<br>Minimum size of a block to be compressed
+`maxWriteQueueLength` | `integer`<br>Maximum amount of data to store in memory. This includes any background writes as well as bundles being built.
+
+## BundleVersion (Enum)
+
+Bundle version number
+
+Name | Description
+---- | -----------
+`Initial` | Initial version number
+`ExportAliases` | Added the BundleExport.Alias property
+`RemoveAliases` | Back out change to include aliases. Will likely do this through an API rather than baked into the data.
+`InPlace` | Use data structures which support in-place reading and writing.
+`ImportHashes` | Add import hashes to imported nodes
+`LatestV1` | Last version using the V1 pipeline
+`PacketSequence` | Structure bundles as a sequence of self-contained packets (uses V2 code)
+`Latest` | The current version number
+`LatestV2` | Last version using the V2 pipeline
+`LatestPlusOne` | Last item in the enum. Used for
+
+## BundleCompressionFormat (Enum)
+
+Indicates the compression format in the bundle
+
+Name | Description
+---- | -----------
+`None` | Packets are uncompressed
+`LZ4` | LZ4 compression
+`Gzip` | Gzip compression
+`Oodle` | Oodle compression (Selkie)
+`Brotli` | Brotli compression
+`Zstd` | ZStandard compression
+
+## ChunkingOptions
+
+Options for creating file nodes
+
+Name | Description
+---- | -----------
+`leafOptions` | [LeafChunkedDataNodeOptions](#leafchunkeddatanodeoptions)<br>Options for creating leaf nodes
+`interiorOptions` | [InteriorChunkedDataNodeOptions](#interiorchunkeddatanodeoptions)<br>Options for creating interior nodes
+
+## LeafChunkedDataNodeOptions
+
+Options for creating a specific type of file nodes
+
+Name | Description
+---- | -----------
+`minSize` | `integer`<br>Minimum chunk size
+`maxSize` | `integer`<br>Maximum chunk size. Chunks will be split on this boundary if another match is not found.
+`targetSize` | `integer`<br>Target chunk size for content-slicing
+`windowSize` | `integer`<br>Window size to use when scanning for split points
+`threshold` | `integer`<br>Accessor for the BuzHash chunking threshold
+
+## InteriorChunkedDataNodeOptions
+
+Options for creating interior nodes
+
+Name | Description
+---- | -----------
+`minChildCount` | `integer`<br>Minimum number of children in each node
+`targetChildCount` | `integer`<br>Target number of children in each node
+`maxChildCount` | `integer`<br>Maximum number of children in each node
+`sliceThreshold` | `integer`<br>Threshold hash value for splitting interior nodes
