[Horde] Removing old plugin references

#Synced-CL 52446693

diff --git a/Engine/Source/Programs/Horde/HordeServer/HordeServer.csproj b/Engine/Source/Programs/Horde/HordeServer/HordeServer.csproj
--- a/Engine/Source/Programs/Horde/HordeServer/HordeServer.csproj
+++ b/Engine/Source/Programs/Horde/HordeServer/HordeServer.csproj
@@ -107,22 +107,10 @@
     <ProjectReference Include="..\..\..\..\Restricted\NotForLicensees\Source\Programs\Horde\EpicPlugins\BuildHealth\HordeServer.Epic.BuildHealth\HordeServer.Epic.BuildHealth.csproj" />
   </ItemGroup>
 
-  <ItemGroup Condition="Exists('$(MSBuildThisFileDirectory)..\..\..\..\Restricted\NotForLicensees\Source\Programs\Horde\EpicPlugins\StructuredAnalytics\HordeServer.StructuredAnalytics\HordeServer.StructuredAnalytics.csproj')">
-    <ProjectReference Include="..\..\..\..\Restricted\NotForLicensees\Source\Programs\Horde\EpicPlugins\StructuredAnalytics\HordeServer.StructuredAnalytics\HordeServer.StructuredAnalytics.csproj" />
-  </ItemGroup>
-
   <ItemGroup Condition="Exists('$(MSBuildThisFileDirectory)..\..\..\..\Restricted\NotForLicensees\Source\Programs\Horde\EpicPlugins\EpicSandbox\HordeServer.Epic.EpicSandbox\HordeServer.Epic.EpicSandbox.csproj')">
     <ProjectReference Include="..\..\..\..\Restricted\NotForLicensees\Source\Programs\Horde\EpicPlugins\EpicSandbox\HordeServer.Epic.EpicSandbox\HordeServer.Epic.EpicSandbox.csproj" />
   </ItemGroup>
 
-  <ItemGroup Condition="Exists('$(MSBuildThisFileDirectory)..\..\..\..\Restricted\NotForLicensees\Source\Programs\Horde\EpicPlugins\ClickHouseTelemetry\HordeServer.Epic.ClickHouseTelemetry\HordeServer.Epic.ClickHouseTelemetry.csproj')">
-    <ProjectReference Include="..\..\..\..\Restricted\NotForLicensees\Source\Programs\Horde\EpicPlugins\ClickHouseTelemetry\HordeServer.Epic.ClickHouseTelemetry\HordeServer.Epic.ClickHouseTelemetry.csproj" />
-  </ItemGroup>
-  
-  <ItemGroup Condition="Exists('$(MSBuildThisFileDirectory)..\..\..\..\Restricted\NotForLicensees\Source\Programs\Horde\EpicPlugins\PerformanceTrends\HordeServer.PerformanceTrends\HordeServer.PerformanceTrends.csproj')">
-    <ProjectReference Include="..\..\..\..\Restricted\NotForLicensees\Source\Programs\Horde\EpicPlugins\PerformanceTrends\HordeServer.PerformanceTrends\HordeServer.PerformanceTrends.csproj" />
-  </ItemGroup>
-
   <ItemGroup>
     <Content Include="..\ThirdParty\**" Condition="$([MSBuild]::IsOSPlatform('Windows'))">
       <CopyToOutputDirectory>Always</CopyToOutputDirectory>
