[Horde][Analytics] Adding in public variants of analytics plugins.

#Synced-CL 52446693

diff --git a/Engine/Source/Programs/Horde/HordeServer/HordeServer.csproj b/Engine/Source/Programs/Horde/HordeServer/HordeServer.csproj
--- a/Engine/Source/Programs/Horde/HordeServer/HordeServer.csproj
+++ b/Engine/Source/Programs/Horde/HordeServer/HordeServer.csproj
@@ -97,8 +97,10 @@
     <ProjectReference Include="..\Plugins\Compute\HordeServer.Compute\HordeServer.Compute.csproj" />
     <ProjectReference Include="..\Plugins\ExamplePlugin\HordeServer.ExamplePlugin\HordeServer.ExamplePlugin.csproj" />
     <ProjectReference Include="..\Plugins\Experimental\HordeServer.Experimental\HordeServer.Experimental.csproj" />
+    <ProjectReference Include="..\Plugins\PerformanceTrends\HordeServer.PerformanceTrends\HordeServer.PerformanceTrends.csproj" />
     <ProjectReference Include="..\Plugins\Secrets\HordeServer.Secrets\HordeServer.Secrets.csproj" />
     <ProjectReference Include="..\Plugins\Storage\HordeServer.Storage\HordeServer.Storage.csproj" />
+    <ProjectReference Include="..\Plugins\StructuredAnalytics\HordeServer.StructuredAnalytics\HordeServer.StructuredAnalytics.csproj" />
     <ProjectReference Include="..\Plugins\Symbols\HordeServer.Symbols\HordeServer.Symbols.csproj" />
     <ProjectReference Include="..\Plugins\Tools\HordeServer.Tools\HordeServer.Tools.csproj" />
   </ItemGroup>
diff --git a/Engine/Source/Programs/Horde/Plugins/PerformanceTrends/HordeServer.PerformanceTrends/HordeServer.PerformanceTrends.csproj b/Engine/Source/Programs/Horde/Plugins/PerformanceTrends/HordeServer.PerformanceTrends/HordeServer.PerformanceTrends.csproj
--- a/Engine/Source/Programs/Horde/Plugins/PerformanceTrends/HordeServer.PerformanceTrends/HordeServer.PerformanceTrends.csproj
+++ b/Engine/Source/Programs/Horde/Plugins/PerformanceTrends/HordeServer.PerformanceTrends/HordeServer.PerformanceTrends.csproj
@@ -1,30 +1,25 @@
 ﻿<Project Sdk="Microsoft.NET.Sdk">
-
   <PropertyGroup>
     <ImplicitUsings>enable</ImplicitUsings>
     <RootNamespace>HordeServer</RootNamespace>
   </PropertyGroup>
-
   <ItemGroup>
     <PackageReference Include="Polly" Version="8.2.1" />
     <PackageReference Include="System.Data.Odbc" Version="9.0.10" />
   </ItemGroup>
-
   <ItemGroup>
-    <ProjectReference Include="..\..\..\..\..\..\..\..\Source\Programs\Shared\EpicGames.Analytics.PerformanceTrends\EpicGames.Analytics.PerformanceTrends.csproj" />
-    <ProjectReference Include="..\..\..\..\..\..\..\..\Source\Programs\Horde\Plugins\Build\HordeServer.Build\HordeServer.Build.csproj" />
-    <ProjectReference Include="..\..\..\..\..\..\..\..\Source\Programs\Shared\EpicGames.Build\EpicGames.Build.csproj" />
-    <ProjectReference Include="..\..\..\..\..\..\..\..\Source\Programs\Shared\EpicGames.Core\EpicGames.Core.csproj" />
-    <ProjectReference Include="..\..\..\..\..\..\..\..\Source\Programs\Shared\EpicGames.Horde\EpicGames.Horde.csproj" />
-    <ProjectReference Include="..\..\..\..\..\..\..\..\Source\Programs\Horde\HordeServer.Shared\HordeServer.Shared.csproj" />
+    <ProjectReference Include="..\..\..\..\Shared\EpicGames.Analytics.PerformanceTrends\EpicGames.Analytics.PerformanceTrends.csproj" />
+    <ProjectReference Include="..\..\Build\HordeServer.Build\HordeServer.Build.csproj" />
+    <ProjectReference Include="..\..\..\..\Shared\EpicGames.Build\EpicGames.Build.csproj" />
+    <ProjectReference Include="..\..\..\..\Shared\EpicGames.Core\EpicGames.Core.csproj" />
+    <ProjectReference Include="..\..\..\..\Shared\EpicGames.Horde\EpicGames.Horde.csproj" />
+    <ProjectReference Include="..\..\..\HordeServer.Shared\HordeServer.Shared.csproj" />
     <ProjectReference Include="..\..\StructuredAnalytics\HordeServer.StructuredAnalytics\HordeServer.StructuredAnalytics.csproj" />
   </ItemGroup>
-
   <PropertyGroup>
     <EmitCompilerGeneratedFiles>true</EmitCompilerGeneratedFiles>
   </PropertyGroup>
   <ItemGroup>
-    <ProjectReference Include="..\..\..\..\..\..\..\..\Source\Programs\Shared\EpicGames.Analytics.Generators\EpicGames.Analytics.Generators.csproj" OutputItemType="Analyzer" ReferenceOutputAssembly="false" />
+    <ProjectReference Include="..\..\..\..\Shared\EpicGames.Analytics.Generators\EpicGames.Analytics.Generators.csproj" OutputItemType="Analyzer" ReferenceOutputAssembly="false" />
   </ItemGroup>
-
-</Project>
+</Project>
\ No newline at end of file
diff --git a/Engine/Source/Programs/Horde/Plugins/StructuredAnalytics/HordeServer.StructuredAnalytics/HordeServer.StructuredAnalytics.csproj b/Engine/Source/Programs/Horde/Plugins/StructuredAnalytics/HordeServer.StructuredAnalytics/HordeServer.StructuredAnalytics.csproj
--- a/Engine/Source/Programs/Horde/Plugins/StructuredAnalytics/HordeServer.StructuredAnalytics/HordeServer.StructuredAnalytics.csproj
+++ b/Engine/Source/Programs/Horde/Plugins/StructuredAnalytics/HordeServer.StructuredAnalytics/HordeServer.StructuredAnalytics.csproj
@@ -1,28 +1,24 @@
 <Project Sdk="Microsoft.NET.Sdk">
-
   <PropertyGroup>
     <ImplicitUsings>enable</ImplicitUsings>
     <RootNamespace>HordeServer</RootNamespace>
   </PropertyGroup>
-
   <ItemGroup>
     <PackageReference Include="Polly" Version="8.2.1" />
     <PackageReference Include="System.Data.Odbc" Version="9.0.10" />
   </ItemGroup>
-
   <ItemGroup>
-    <ProjectReference Include="..\..\..\..\..\..\..\..\Source\Programs\Horde\Plugins\Build\HordeServer.Build\HordeServer.Build.csproj" />
-    <ProjectReference Include="..\..\..\..\..\..\..\..\Source\Programs\Shared\EpicGames.Build\EpicGames.Build.csproj" />
-    <ProjectReference Include="..\..\..\..\..\..\..\..\Source\Programs\Shared\EpicGames.Core\EpicGames.Core.csproj" />
-    <ProjectReference Include="..\..\..\..\..\..\..\..\Source\Programs\Shared\EpicGames.Horde\EpicGames.Horde.csproj" />
-    <ProjectReference Include="..\..\..\..\..\..\..\..\Source\Programs\Horde\HordeServer.Shared\HordeServer.Shared.csproj" />
-  </ItemGroup>
+    <ProjectReference Include="..\..\..\..\Shared\EpicGames.Build\EpicGames.Build.csproj" />
+    <ProjectReference Include="..\..\..\..\Shared\EpicGames.Core\EpicGames.Core.csproj" />
+    <ProjectReference Include="..\..\..\..\Shared\EpicGames.Horde\EpicGames.Horde.csproj" />
+    <ProjectReference Include="..\..\..\HordeServer.Shared\HordeServer.Shared.csproj" />
 
+  </ItemGroup>
   <PropertyGroup>
     <EmitCompilerGeneratedFiles>true</EmitCompilerGeneratedFiles>
   </PropertyGroup>
   <ItemGroup>
-    <ProjectReference Include="..\..\..\..\..\..\..\..\Source\Programs\Shared\EpicGames.Analytics.Generators\EpicGames.Analytics.Generators.csproj" OutputItemType="Analyzer" ReferenceOutputAssembly="false" />
+    <ProjectReference Include="..\..\..\..\Shared\EpicGames.Analytics.Generators\EpicGames.Analytics.Generators.csproj" OutputItemType="Analyzer" ReferenceOutputAssembly="false" />
+    <ProjectReference Include="..\..\Build\HordeServer.Build\HordeServer.Build.csproj" />
   </ItemGroup>
-
-</Project>
+</Project>
\ No newline at end of file
