[Horde][Analytics] Public variants of performance trends and structured analytics

#Synced-CL 52533234

diff --git a/Engine/Source/Programs/Horde/HordeDashboard/plugins/performancetrends/PerformanceTrendView.tsx b/Engine/Source/Programs/Horde/HordeDashboard/plugins/performancetrends/PerformanceTrendView.tsx
new file mode 100644
--- /dev/null
+++ b/Engine/Source/Programs/Horde/HordeDashboard/plugins/performancetrends/PerformanceTrendView.tsx
@@ -0,0 +1,400 @@
+// Copyright Epic Games, Inc. All Rights Reserved.
+
+import { Checkbox, Dialog, DialogType, DirectionalHint, Dropdown, IDropdownOption, Label, PrimaryButton, Spinner, SpinnerSize, Stack, TooltipHost } from "@fluentui/react"
+import { observer } from "mobx-react-lite";
+import React, { useEffect, useLayoutEffect, useMemo, useRef, useState } from "react";
+import { getHordeStyling } from "horde/styles/Styles";
+import { TopNav } from "horde/components/TopNav";
+import { Breadcrumbs } from "horde/components/Breadcrumbs";
+import { useWindowSize } from "horde/base/utilities/hooks";
+import { Text } from '@fluentui/react';
+import { TestProjectDropdownSingle } from "./components/TestProjectDropdownComponent";
+import { runInAction } from "mobx";
+import { CommitDropdownSingle } from "./components/CommitDropdownComponent";
+import { MetricSelectionHandler } from "./components/LineGraphComponent";
+import { MetricSummaryTypeDropdownSingle } from "./components/MetricSummaryTypeDropdownSingle";
+import { PerformanceTrendUrlSync } from "./components/PerformanceTrendUrl";
+import { PlatformDropdownMulti } from "./components/PlatformDropdownComponent";
+import { StreamDropdownMulti } from "./components/StreamDropdownComponent";
+import { TestIdentitiesTypeDropdownSingle } from "./components/TestIdentitiesTypeDropdownSingle";
+import { TestTypeDropdownMulti } from "./components/TestTypeDropdownComponent";
+import { ViewablePropertiesDropdown } from "./components/ViewablePropertiesDropdownComponent";
+import { PerformanceTrendOptionsController } from "./filters/PerformanceTrendOptionsController";
+import { PerformanceTrendOptionsDataHandler } from "./filters/PerformanceTrendOptionsDataHandler";
+import { PerformanceTrendOptionsState } from "./filters/PerformanceTrendOptionsState";
+import { GroupByOption, PerformanceTrendUIOptionsState } from "./filters/PerformanceTrendUIOptionsState";
+import { KeyStatsData } from "./metrictypes/KeyStatsData";
+import { MetricConstraint, NO_RESOLVE_TYPE_FOUND, MetricTypeRegistry, PerformanceTrendContext } from "./metrictypes/PerformanceTrendsTypes";
+import { PerformanceTrendDataHandler, PerformanceTrendFilter } from "./PerformanceTrendsDataHandler";
+import { decodeMetricTypeKey } from "./responses/FilterKeys";
+import { IMetricViewGenerator, IMetricDetailedView, IMetricView } from "./viewgenerators/PerformanceTrendRenderTypes";
+
+// #region -- Visual Components --
+
+const DEFAULT_RENDER_COMPONENT_WIDTH = 900;
+const DEFAULT_RENDER_COMPONENT_HEIGHT = 500;
+
+function selectDataProviderAndViewGenerator(performanceTrendOptions: PerformanceTrendOptionsController): { dataProvider: PerformanceTrendDataHandler<MetricConstraint> | undefined, viewGenerator: IMetricViewGenerator<MetricConstraint> | undefined } {
+    const metricType = performanceTrendOptions.state.publishedEnabledMetricSummaryTypes.length === 0 ? NO_RESOLVE_TYPE_FOUND : decodeMetricTypeKey(performanceTrendOptions.state.publishedEnabledMetricSummaryTypes[0])?.summaryType ?? NO_RESOLVE_TYPE_FOUND;
+
+    return MetricTypeRegistry.get(metricType);
+}
+
+const PerformanceTrendSidebar: React.FC<{ performanceTrendOptions: PerformanceTrendOptionsController }> = observer(function ConstructPerformanceTrendsSidebar({ performanceTrendOptions: performanceTrendOptions }) {
+    const { hordeClasses, modeColors } = getHordeStyling();
+
+    // #region -- UI Controls --
+
+    const [curveGraphLine, setCurveGraphLine] = useState(() => performanceTrendOptions.uiState.curveGraphLine)
+    const [groupByOption, setGroupByOption] = useState<GroupByOption>(() => performanceTrendOptions.uiState.groupByOption);
+    const { dataProvider, viewGenerator } = selectDataProviderAndViewGenerator(performanceTrendOptions);
+
+    const groupByOptions: IDropdownOption<GroupByOption>[] = [
+        { key: 'platform', text: 'Platform', data: 'platform' },
+        { key: 'stream', text: 'Stream', data: 'stream' },
+        { key: 'testType', text: 'Test Type', data: 'testType' },
+        { key: 'platform+stream', text: 'Platform + Stream', data: 'platform+stream' },
+        { key: 'platform+testType', text: 'Platform + Test Type', data: 'platform+testType' },
+        { key: 'stream+testType', text: 'Stream + Test Type', data: 'stream+testType' },
+        { key: 'platform+stream+testType', text: 'Platform + Stream + Test Type', data: 'platform+stream+testType' },
+    ];
+
+    const updateCurveGraphLine = (isChecked) => {
+        setCurveGraphLine(isChecked);
+        performanceTrendOptions.setCurveGraphLine(isChecked);
+    }
+
+    const updateGroupByOption = (option: GroupByOption) => {
+        setGroupByOption(option);
+        performanceTrendOptions.setGroupByOption(option);
+    }
+
+    useEffect(() => {
+        setCurveGraphLine(performanceTrendOptions.uiState.curveGraphLine);
+    }, [performanceTrendOptions.uiState.curveGraphLine]);
+
+    useEffect(() => {
+        setGroupByOption(performanceTrendOptions.uiState.groupByOption);
+    }, [performanceTrendOptions.uiState.groupByOption]);
+
+    // #endregion -- UI Controls --
+
+    return (
+        <Stack style={{ minWidth: 375, maxWidth: 500, paddingRight: 18 }}>
+            <Stack className={hordeClasses.modal} tokens={{ childrenGap: 12 }}>
+                <Stack
+                    styles={{
+                        root: {
+                            border: `1px solid ${modeColors.content}`,
+                            borderRadius: 4,
+                            padding: 8,
+                        },
+                    }}
+                >
+                    <Stack style={{ paddingTop: 0, paddingBottom: 4 }}>
+                        <Label>Performance Trend Filters</Label>
+                    </Stack>
+                    <Stack
+                        tokens={{ childrenGap: 12 }}
+                        styles={{
+                            root: {
+                                border: `1px solid ${modeColors.content}`,
+                                borderRadius: 4,
+                                padding: 8,
+                            },
+                        }}>
+                        <TestProjectDropdownSingle handler={optionsHandler} performanceTrendOptions={performanceTrendOptions} />
+                        <TestIdentitiesTypeDropdownSingle handler={optionsHandler} performanceTrendOptions={performanceTrendOptions} />
+                        <MetricSummaryTypeDropdownSingle handler={optionsHandler} performanceTrendOptions={performanceTrendOptions} />
+                        <TestTypeDropdownMulti handler={optionsHandler} performanceTrendOptions={performanceTrendOptions} />
+                        <StreamDropdownMulti handler={optionsHandler} performanceTrendOptions={performanceTrendOptions} />
+                        <PlatformDropdownMulti handler={optionsHandler} performanceTrendOptions={performanceTrendOptions} />
+                        <CommitDropdownSingle handler={optionsHandler} performanceTrendOptions={performanceTrendOptions} />
+                        <TooltipHost content={
+                            <Stack style={{ maxWidth: 300 }}>
+                                <Stack>
+                                    Clears the current Project / Test Identity / Summary Type/ Commit Range selection.
+                                </Stack>
+                            </Stack>
+                        }
+                            calloutProps={{
+                                directionalHint: DirectionalHint.rightCenter
+                            }}
+                        >
+                            <PrimaryButton text="Clear" onClick={() => {
+                                runInAction(() => {
+                                    performanceTrendOptions.getTransactionSession().clearAll();
+                                    optionsHandler.reset();
+                                    performanceTrendOptions.commitTransactionSession();
+                                });
+                            }} />
+                        </TooltipHost>
+                    </Stack>
+                    <Stack
+                        styles={{
+                            root: {
+                                border: `1px solid ${modeColors.content}`,
+                                borderRadius: 4,
+                                padding: 8,
+                            },
+                        }}>
+                        <Stack style={{ paddingTop: 0, paddingBottom: 4 }}>
+                            <Label>Options</Label>
+                        </Stack>
+                        <ViewablePropertiesDropdown viewGenerator={viewGenerator} handler={optionsHandler} performanceTrendOptions={performanceTrendOptions} />
+                        <TooltipHost content={
+                            <div style={{ maxWidth: 300 }}>
+                                <div>
+                                    Draw line graphs with curve-interpolated lines.
+                                </div>
+                            </div>
+                        }
+                            calloutProps={{
+                                directionalHint: DirectionalHint.rightCenter
+                            }}
+                        >
+                            <Checkbox
+                                label="Curved Lines"
+                                checked={curveGraphLine}
+                                onChange={(_, isChecked) => updateCurveGraphLine(!!isChecked)}
+                            />
+                        </TooltipHost>
+                        <TooltipHost content={
+                            <div style={{ maxWidth: 300 }}>
+                                <div>
+                                    Group graph series by the selected dimensions.
+                                </div>
+                            </div>
+                        }
+                            calloutProps={{
+                                directionalHint: DirectionalHint.rightCenter
+                            }}
+                        >
+                            <Dropdown
+                                label="Group By"
+                                selectedKey={groupByOption}
+                                options={groupByOptions}
+                                onChange={(_, option) => {
+                                    if (option?.data) {
+                                        updateGroupByOption(option.data);
+                                    }
+                                }}
+                            />
+                        </TooltipHost>
+                    </Stack>
+                </Stack>
+            </Stack>
+        </Stack>
+    )
+});
+
+interface AnalyticsGraphClusterProps {
+    dataProvider: PerformanceTrendDataHandler<MetricConstraint>;
+    viewGenerator: IMetricViewGenerator<MetricConstraint>;
+    performanceTrendOptions: PerformanceTrendOptionsController;
+    selectionHandler?: MetricSelectionHandler<PerformanceTrendContext>;
+}
+
+type MetricDetailModalProps = {
+    context: PerformanceTrendContext[];
+    performanceTrendOptions: PerformanceTrendOptionsController;
+    viewGenerator: IMetricViewGenerator<MetricConstraint>;
+    onDismiss: () => void;
+};
+
+const MetricDetailModal: React.FC<MetricDetailModalProps> = ({ context, performanceTrendOptions, viewGenerator, onDismiss }) => {
+    const view: IMetricDetailedView = useMemo(() => viewGenerator.getSelectedDetailView(context, performanceTrendOptions.uiState), [context, viewGenerator, performanceTrendOptions.uiOptionsChangeVersion]);
+    const Render = view.Render;
+    return (
+        <Dialog
+            modalProps={{
+                isBlocking: false,
+                topOffsetFixed: true,
+                styles: {
+                    root: {
+                        selectors: {
+                            ".ms-Dialog-title": {
+                                paddingTop: '24px',
+                                paddingLeft: '32px'
+                            }
+                        }
+                    }
+                }
+            }}
+            onDismiss={onDismiss}
+            hidden={false}
+            minWidth={1200}
+            dialogContentProps={{
+                type: DialogType.close,
+                onDismiss: onDismiss,
+                title: "Metric Detail",
+            }}>
+            <Stack tokens={{ childrenGap: 8 }}>
+                <Render
+                    width={DEFAULT_RENDER_COMPONENT_WIDTH}
+                    height={DEFAULT_RENDER_COMPONENT_HEIGHT}
+                />
+            </Stack>
+        </Dialog >
+    );
+};
+
+const AnalyticsGraphCluster: React.FC<AnalyticsGraphClusterProps> = observer(({ dataProvider, viewGenerator, performanceTrendOptions, selectionHandler }) => {
+    const containerRef = useRef<HTMLDivElement>(null);
+    const [containerWidth, setContainerWidth] = useState(0);
+
+    useLayoutEffect(() => {
+        if (!containerRef.current) return;
+
+        const resizeObserver = new ResizeObserver((entries) => {
+            for (const entry of entries) {
+                const newWidth = entry.contentRect.width;
+                // Only update if width changed by more than 1px to prevent flickering
+                // from sub-pixel layout shifts during tooltip interactions
+                setContainerWidth(prev => {
+                    return Math.abs(prev - newWidth) > 1 ? newWidth : prev
+                }
+                );
+            }
+        });
+
+        resizeObserver.observe(containerRef.current);
+        return () => resizeObserver.disconnect();
+    }, []);
+
+    // Note: this is by design where we are monitoring data, the ui options change version.
+    // The rationale is as follows:
+    //  - This is a primary extension point for licensees to add new view generators
+    //  - We never know what downstream view generators will have as a UI requirement (this struct will evolve, and implement new options)
+    //  - We do not know if it will be observed properly (derived components in getViews() impl may not be observed, as an example)
+    //  - Any closures within the getViews delegate (where observed views are modulated) outside of the observable section will un-observe parameters
+    //  - This is the safest way to guarantee that we perform a redraw
+    const views: IMetricView[] = useMemo(() => viewGenerator.getViews(dataProvider.data, performanceTrendOptions.uiState, selectionHandler), [viewGenerator, dataProvider.data, performanceTrendOptions.uiOptionsChangeVersion]);
+
+    const columnWidth = containerWidth ? Math.floor((containerWidth - 12) / 2) : 300;
+
+    return (
+        <div ref={containerRef} style={{ width: "100%" }}>
+            <Stack horizontal wrap tokens={{ childrenGap: 12 }}>
+                {views.map((view, index) => (
+                    <Stack key={`${KeyStatsData.name}${index}`} styles={{ root: { width: columnWidth, minWidth: 600 } }}>
+                        <view.Render width={columnWidth} height={DEFAULT_RENDER_COMPONENT_HEIGHT} />
+                    </Stack>
+                ))}
+            </Stack>
+        </div>
+    );
+});
+
+const PerformanceTrendComponent: React.FC<{ performanceTrendOptions: PerformanceTrendOptionsController, filter: PerformanceTrendFilter }> = observer(({ performanceTrendOptions, filter }) => {
+    const { hordeClasses, modeColors } = getHordeStyling();
+    const [selectedMetric, setSelectedMetric] = useState<PerformanceTrendContext[] | null>(null);
+    const { dataProvider, viewGenerator } = selectDataProviderAndViewGenerator(performanceTrendOptions);
+
+    const handleMetricSelection: MetricSelectionHandler<PerformanceTrendContext> = (context) => {
+        setSelectedMetric(context);
+    };
+
+    useEffect(() => {
+        if (dataProvider) {
+            dataProvider.setFilter(filter);
+        }
+    }, [filter, dataProvider]);
+
+    if (!viewGenerator || !dataProvider) {
+        return null;
+    }
+
+    if (dataProvider.querying) {
+        return <Spinner style={{ paddingTop: 16 }} size={SpinnerSize.large} />;
+    }
+
+    return dataProvider.data.length === 0 ? (
+        <Text
+            variant="medium"
+            styles={{ root: { paddingTop: 16, color: modeColors.textSecondary } }}
+        >
+            No data available for the selected filters.
+        </Text>) : (
+        <Stack styles={{ root: { width: "100%" } }} tokens={{ childrenGap: 24 }}>
+            {selectedMetric && (<MetricDetailModal context={selectedMetric} viewGenerator={viewGenerator} performanceTrendOptions={performanceTrendOptions} onDismiss={() => setSelectedMetric(null)} />)}
+            <AnalyticsGraphCluster dataProvider={dataProvider} viewGenerator={viewGenerator} performanceTrendOptions={performanceTrendOptionsController} selectionHandler={handleMetricSelection} />
+            <Stack></Stack>
+            {/*<CSVTable dataProvider={dataProvider} columnProvider={dataProvider} />*/}
+        </Stack>
+    );
+});
+
+/**
+ * Horde Performance Trends view.
+ * @returns React component.
+ */
+export const PerformanceTrendsView: React.FC = observer(() => {
+    const { hordeClasses, modeColors } = getHordeStyling();
+    const windowSize = useWindowSize();
+
+    const filter: PerformanceTrendFilter = useMemo(() => {
+        return {
+            testProjects: performanceTrendOptionsState.publishedEnabledProjectKeys,
+            testIdentities: performanceTrendOptionsState.publishedEnabledTestIdentityKeys,
+            testTypes: performanceTrendOptionsState.publishedEnabledTestTypeKeys,
+            metricTypes: performanceTrendOptionsState.publishedEnabledMetricSummaryTypes,
+            streams: performanceTrendOptionsState.publishedEnabledStreamKeys,
+            platforms: performanceTrendOptionsState.publishedEnabledPlatforms,
+            startCommit: performanceTrendOptionsState.publishedEnabledStartCommit,
+            endCommit: performanceTrendOptionsState.publishedEnabledEndCommit
+        };
+    }, [
+        performanceTrendOptionsState.publishedEnabledProjectKeys,
+        performanceTrendOptionsState.publishedEnabledTestIdentityKeys,
+        performanceTrendOptionsState.publishedEnabledTestTypeKeys,
+        performanceTrendOptionsState.publishedEnabledMetricSummaryTypes,
+        performanceTrendOptionsState.publishedEnabledStreamKeys,
+        performanceTrendOptionsState.publishedEnabledPlatforms,
+        performanceTrendOptionsState.publishedEnabledStartCommit,
+        performanceTrendOptionsState.publishedEnabledEndCommit
+    ]);
+
+    return (
+        <Stack className={hordeClasses.horde}>
+            <PerformanceTrendUrlSync buildHealthController={performanceTrendOptionsController} handler={optionsHandler} />
+            <TopNav />
+            <Breadcrumbs items={[{ text: 'Performance Trends' }]} />
+            <Stack horizontal>
+                <div key={`windowsize_performance_trends_view_${windowSize.width}_${windowSize.height}`} style={{ width: 0, flexShrink: 0, backgroundColor: modeColors.background }} />
+                <Stack horizontalAlign="center" grow styles={{ root: { width: "100%", padding: 12, backgroundColor: modeColors.background } }}>
+                    <Stack styles={{ root: { width: "100%" } }}>
+                        <Stack horizontal styles={{ root: { minHeight: '85vh', maxHeight: '85vh' } }} >
+                            <PerformanceTrendSidebar performanceTrendOptions={performanceTrendOptionsController} />
+                            <Stack id="parent-performance-trends" styles={{
+                                root: {
+                                    overflowX: "hidden",
+                                    overflowY: "auto",
+                                    minWidth: 0,
+                                    position: "relative",
+                                    border: `${modeColors.content} solid 2px`,
+                                    display: "flex",
+                                    flexDirection: "column",
+                                    flex: 1,
+                                }
+                            }}>
+                                <PerformanceTrendComponent performanceTrendOptions={performanceTrendOptionsController} filter={filter} />
+                            </Stack>
+                        </Stack>
+                    </Stack>
+                </Stack>
+            </Stack>
+        </Stack>
+    )
+});
+
+// #endregion -- Visual Components --
+
+// #region -- Module Scoped Variables --
+
+const performanceTrendOptionsState = new PerformanceTrendOptionsState();
+const performanceTrendUIOptionsState = new PerformanceTrendUIOptionsState();
+const performanceTrendOptionsController = new PerformanceTrendOptionsController(performanceTrendOptionsState, performanceTrendUIOptionsState);
+const optionsHandler = new PerformanceTrendOptionsDataHandler(performanceTrendOptionsState, performanceTrendOptionsController);
+
+// #endregion -- Module Scoped Variables --
\ No newline at end of file
diff --git a/Engine/Source/Programs/Horde/HordeDashboard/plugins/performancetrends/PerformanceTrendsDataHandler.ts b/Engine/Source/Programs/Horde/HordeDashboard/plugins/performancetrends/PerformanceTrendsDataHandler.ts
new file mode 100644
--- /dev/null
+++ b/Engine/Source/Programs/Horde/HordeDashboard/plugins/performancetrends/PerformanceTrendsDataHandler.ts
@@ -0,0 +1,173 @@
+// Copyright Epic Games, Inc. All Rights Reserved.
+
+import { action, computed, makeObservable, observable, runInAction } from "mobx";
+import { ISummaryMetricProvider, MetricConstraint } from "./metrictypes/PerformanceTrendsTypes";
+import { DecodedTestIdentityKey, decodeTestIdentity, DecodedTestTypeKey, decodeTestTypeKey, decodeStreamKey, DecodedStreamKey, decodePlatformKey, DecodedPlatformKey } from "./responses/FilterKeys";
+import { IDataProvider, IColumnProvider } from "hordePlugins/structuredanalytics/CSVComponent";
+import { getMetrics } from "./api";
+
+/**
+ * The filter interface supported by the @see PerformanceTrendsDataHandler .
+ */
+export interface PerformanceTrendFilter {
+    testProjects: string[];
+    testIdentities: string[];
+    testTypes: string[];
+    metricTypes: string[];
+    streams: string[];
+    platforms: string[];
+    startCommit: number | null;
+    endCommit: number | null;
+}
+
+// #region -- Data Handler --
+
+/**
+ * General purpose data handler for metric retrieval, based on types.
+ */
+export class PerformanceTrendDataHandler<TMetric extends MetricConstraint> implements IDataProvider, IColumnProvider {
+
+    // #region -- Private Members --
+
+    private ctor: ISummaryMetricProvider<TMetric>
+    private activeFilter: PerformanceTrendFilter | null = null;
+    private performanceReportTelemetryData: TMetric[] = [];
+
+    @observable
+    private activeRequests = 0;
+
+    // #endregion -- Private Members --
+
+    // #region -- Public Members --
+
+    // IColumnProvider
+    @observable
+    columns: string[] = [];
+
+    // IDataProvider
+    @observable
+    data: TMetric[] = [];
+
+    // Query state
+    @computed
+    get querying(): boolean {
+        return this.activeRequests > 0;
+    }
+    // #endregion -- Public Members --
+
+    // #region -- Constructor --
+
+    constructor(ctor: ISummaryMetricProvider<TMetric>) {
+        makeObservable(this);
+        this.ctor = ctor;
+        this.columns = [];
+        this.data = [];
+    }
+
+    // #endregion -- Constructor --
+
+    // #region -- Public Api --
+
+    /**
+     * Queries the underlying database to populate the handler with performance trend data.
+     */
+    async query() {
+        if (!this.activeFilter || !this.activeFilter.testIdentities || !this.activeFilter.testProjects || this.activeFilter.testIdentities.length === 0 || this.activeFilter.testTypes.length === 0 || this.activeFilter.testProjects.length === 0 || this.activeFilter.streams.length === 0 || this.activeFilter.platforms.length === 0) {
+            this.data = [];
+
+            return;
+        }
+
+        let targetTestProject: string;
+
+        // We will just support one test project for now.
+        if (this.activeFilter.testProjects.length > 0) {
+            targetTestProject = this.activeFilter.testProjects[0];
+        } else {
+            console.warn("[PerformanceTrendsDataHandler] Received no test project.");
+            return;
+        }
+
+        let targetTestIdentity: string;
+        let targetTestTypes: string[] = [];
+        let targetStreams: string[] = [];
+        let targetPlatforms: string[] = [];
+
+        // We will just support one test identity for now.
+        if (this.activeFilter.testIdentities.length > 0) {
+            let decodedTestIdentity: DecodedTestIdentityKey | null = decodeTestIdentity(this.activeFilter.testIdentities[0]);
+            if (decodedTestIdentity === null) {
+                console.warn(`[PerformanceTrendDataHandler] Handler received an invalid test identity option: ${this.activeFilter.testIdentities[0]}`);
+                return;
+            }
+
+            targetTestIdentity = decodedTestIdentity.testIdentity;
+        } else {
+            console.warn("[PerformanceTrendsDataHandler] Received no test identity.");
+            return;
+        }
+
+        // We will just support one test type for now.
+
+        if (this.activeFilter.testTypes.length > 0) {
+            targetTestTypes = this.activeFilter.testTypes
+                .map(x => decodeTestTypeKey(x))
+                .filter((x: DecodedTestTypeKey | null): x is DecodedTestTypeKey => x !== null)
+                .map(x => x.testType);
+        }
+
+        if (this.activeFilter.streams.length > 0) {
+            targetStreams = this.activeFilter.streams
+                .map(x => decodeStreamKey(x))
+                .filter((x: DecodedStreamKey | null): x is DecodedStreamKey => x !== null)
+                .map(x => x.stream);
+        }
+
+        if (this.activeFilter.platforms.length > 0) {
+            targetPlatforms = this.activeFilter.platforms
+                .map(x => decodePlatformKey(x))
+                .filter((x: DecodedPlatformKey | null): x is DecodedPlatformKey => x !== null)
+                .map(x => x.platform);
+        }
+
+        runInAction(() => {
+            this.activeRequests++;
+        });
+
+        try {
+            this.performanceReportTelemetryData = await getMetrics<TMetric>(this.ctor, targetTestProject, targetTestIdentity, targetTestTypes, targetStreams, targetPlatforms, this.activeFilter.startCommit, this.activeFilter.endCommit);
+        } catch (errorReason) {
+            console.log(errorReason);
+
+            runInAction(() => {
+                this.activeRequests--;
+            });
+
+            return;
+        }
+
+        runInAction(() => {
+            this.activeRequests--;
+        });
+
+        this.data = this.performanceReportTelemetryData;
+        this.columns =
+            this.performanceReportTelemetryData.length >= 1
+                ? Object.keys(this.performanceReportTelemetryData[0])
+                : [];
+    }
+
+    /**
+     * Sets a new filter for the handler, which will control the underlying dataset.
+     * @param filter The new filter to use.
+     */
+    @action
+    setFilter(filter: PerformanceTrendFilter | null) {
+        this.activeFilter = filter;
+        this.query();
+    }
+
+    // #endregion -- Public Api --
+}
+
+// #endregion
\ No newline at end of file
diff --git a/Engine/Source/Programs/Horde/HordeDashboard/plugins/performancetrends/api.ts b/Engine/Source/Programs/Horde/HordeDashboard/plugins/performancetrends/api.ts
new file mode 100644
--- /dev/null
+++ b/Engine/Source/Programs/Horde/HordeDashboard/plugins/performancetrends/api.ts
@@ -0,0 +1,143 @@
+import backend from "horde/backend";
+import { REASON_REENTRANT_REQUEST, RequestCancellationManager } from "horde/backend/RequestCancellationManager";
+import { MetricConstraint, ISummaryMetricProvider, MetricTypeRegistry } from "./metrictypes/PerformanceTrendsTypes";
+import { GetTestProjectResponse, GetTestProjectPlatformResponse, GetTestProjectCommitResponse } from "./responses/GetTestProjectResponse";
+
+let performanceTrendsCancellationManager: RequestCancellationManager = new RequestCancellationManager();
+
+export const DEFAULT_METRIC_COUNT: number = 2000;
+
+/**
+ * Gets a metric of a specific @see ISummaryMetricProvider<T> type.
+ * @param metricProvider The metric provider to use for request, and construction of resulting data.
+ * @param testProject The test project to obtain data for.
+ * @param testIdentity The test identity of the metric to filter for.
+ * @param testType The test type of the metric to filter for.
+ * @param streams The streams to filter for.
+ * @param platforms The platforms to filter for.
+ * @returns An array of constructed metrics.
+ */
+export async function getMetrics<T extends MetricConstraint>(metricProvider: ISummaryMetricProvider<T>, testProject: string, testIdentity: string, testTypes: string[], streams?: string[], platforms?: string[], startCommit: number | null = null, endCommit: number | null = null, count: number = DEFAULT_METRIC_COUNT): Promise<T[]> {
+    let params = { type: metricProvider.metricType };
+
+    params["testProject"] = testProject;
+    params["testIdentity"] = testIdentity;
+    params["testTypes"] = testTypes;
+    params["count"] = count;
+
+    if (streams !== undefined) {
+        params["streams"] = streams;
+    }
+
+    if (platforms !== undefined) {
+        params["platforms"] = platforms;
+    }
+
+    if (startCommit !== null) {
+        params["startCommitIdOrdered"] = startCommit;
+    }
+
+    if (endCommit !== null) {
+        params["endCommitIdOrdered"] = endCommit;
+    }
+
+    const response = await backend.fetch.get(`/api/v1/buildhealth/performancetrends/metrics`,
+        {
+            params,
+            signal: performanceTrendsCancellationManager.getSignal("getMetrics")
+        }
+    );
+
+    return (response.data as any[]).map(x => MetricTypeRegistry.create<T>(metricProvider.metricType, x)).filter((x): x is T => x != null);
+}
+
+/**
+ * Gets all test projects, across all known metric types.
+ * @returns A list of test projects across all known metric types.
+ */
+export async function getTestProjects(): Promise<GetTestProjectResponse[]> {
+    return await backend.fetch.get(`/api/v1/buildhealth/performancetrends/testprojects`,
+        {
+            signal: performanceTrendsCancellationManager.getSignal("getTestProjects")
+        }
+    ).then((response) => {
+        let data = response.data as any;
+        const results: GetTestProjectResponse[] = data.map((x) => ({
+            summaryType: x["summaryType"],
+            testName: x["testName"],
+            testIdentity: x["testIdentity"],
+            testType: x["testType"],
+            computedStream: x["computedStream"]
+        }));
+
+        return results;
+    }).catch(reason => {
+        if (String(reason).includes(REASON_REENTRANT_REQUEST)) {
+            return [];
+        }
+        throw reason;
+    });
+}
+
+/**
+ * Gets all platforms given the test project, test identity, and summary type.
+ * @returns A list of platforms for the given test project, test identity and summary type.
+ */
+export async function getPlatforms(testProject: string, testIdentity: string, metricSummaryType: string, streams?: string[]): Promise<GetTestProjectPlatformResponse[]> {
+    let params = {};
+
+    params["metricSummaryType"] = metricSummaryType;
+    params["testProject"] = testProject;
+    params["testIdentity"] = testIdentity;
+
+    if (streams) {
+        params["streams"] = streams;
+    }
+
+    return await backend.fetch.get(`/api/v1/buildhealth/performancetrends/platforms`,
+        {
+            params: { ...params },
+            signal: performanceTrendsCancellationManager.getSignal("getPlatforms")
+        }
+    ).then(response => response.data as GetTestProjectPlatformResponse[])
+        .catch(reason => {
+            if (String(reason).includes(REASON_REENTRANT_REQUEST)) {
+                return [];
+            }
+            throw reason;
+        });
+}
+
+/**
+ * Gets all commits given the test project, test identity, summary type, and platform.
+ * @returns A list of commits for the given test project, test identity, summary type, and platform.
+ */
+export async function getCommits(testProject: string, testIdentity: string, metricSummaryType: string, streams: string[], platforms: string[]): Promise<GetTestProjectCommitResponse[]> {
+    let params = {};
+
+    params["metricSummaryType"] = metricSummaryType;
+    params["testProject"] = testProject;
+    params["testIdentity"] = testIdentity;
+
+    if (streams.length > 0) {
+        params["streams"] = streams;
+    }
+
+    if (platforms.length > 0) {
+        params["platforms"] = platforms;
+    }
+
+    return backend.fetch.get(
+        `/api/v1/buildhealth/performancetrends/commits`,
+        {
+            params,
+            signal: performanceTrendsCancellationManager.getSignal("getCommits")
+        }
+    ).then(response => response.data as GetTestProjectCommitResponse[])
+        .catch(reason => {
+            if (String(reason).includes(REASON_REENTRANT_REQUEST)) {
+                return [];
+            }
+            throw reason;
+        });
+}
\ No newline at end of file
diff --git a/Engine/Source/Programs/Horde/HordeDashboard/plugins/performancetrends/components/CommitDropdownComponent.tsx b/Engine/Source/Programs/Horde/HordeDashboard/plugins/performancetrends/components/CommitDropdownComponent.tsx
new file mode 100644
--- /dev/null
+++ b/Engine/Source/Programs/Horde/HordeDashboard/plugins/performancetrends/components/CommitDropdownComponent.tsx
@@ -0,0 +1,110 @@
+// Copyright Epic Games, Inc. All Rights Reserved.
+
+import { Stack, Label, IDropdownOption, SpinnerSize, Spinner } from "@fluentui/react";
+import { observer } from "mobx-react-lite";
+import { PerformanceTrendOptionsController } from "../filters/PerformanceTrendOptionsController";
+import { EMPTY_REFRESH_REQUEST, PerformanceTrendOptionsDataHandler } from "../filters/PerformanceTrendOptionsDataHandler";
+import { PerformanceTrendOptionList, PerformanceTrendOptionData, BasicListParameter } from "./DropdownComponent";
+
+/**
+ * Component used to model all known commits (for the currently selected test project, test identity, summary table, and platform), as a single select dropdown.
+ * @return The react component.
+ */
+export const CommitDropdownSingle: React.FC<{ handler: PerformanceTrendOptionsDataHandler, performanceTrendOptions: PerformanceTrendOptionsController }> = observer(function ConstructProjectDropdownSingle({ handler, performanceTrendOptions }) {
+    if (handler.isLoadingCommits) {
+        return (
+            <Stack>
+                <Label>Commit Range</Label>
+                <Spinner size={SpinnerSize.small} label="Loading commits for project..." />
+            </Stack>
+        );
+    }
+
+    const selectedStartCommit = performanceTrendOptions.state.enabledStartCommit;
+    const selectedEndCommit = performanceTrendOptions.state.enabledEndCommit;
+
+    // Build all commit items sorted numerically (flatten arrays from each stream response)
+    const seenCommits = new Set<number>();
+    const allCommitItems: { id: string; text: string; value: number }[] = [];
+
+    for (const commitResponse of handler.commits) {
+        for (const commitId of commitResponse.commitIds) {
+            if (!seenCommits.has(commitId)) {
+                seenCommits.add(commitId);
+                allCommitItems.push({
+                    id: commitId.toFixed(),
+                    text: commitId.toFixed(),
+                    value: commitId
+                });
+            }
+        }
+    }
+
+    allCommitItems.sort((a, b) => a.value - b.value);
+
+    // Filter start commit options: if end commit is selected, only show commits <= end commit
+    const startCommitItems = selectedEndCommit
+        ? allCommitItems.filter(item => item.value <= selectedEndCommit)
+        : allCommitItems;
+
+    // Filter end commit options: if start commit is selected, only show commits >= start commit
+    const endCommitItems = selectedStartCommit
+        ? allCommitItems.filter(item => item.value >= selectedStartCommit)
+        : allCommitItems;
+
+    const startDropDownParams: PerformanceTrendOptionList = {
+        label: "",
+        items: startCommitItems,
+        tooltip: "Select the start commit for the range."
+    };
+
+    const endDropDownParams: PerformanceTrendOptionList = {
+        label: "",
+        items: endCommitItems,
+        tooltip: "Select the end commit for the range."
+    };
+
+    const selectedStartKey = selectedStartCommit?.toFixed();
+    const selectedEndKey = selectedEndCommit?.toFixed();
+
+    return allCommitItems.length > 0 && (
+        <Stack tokens={{ childrenGap: 8 }}>
+            <Label>Commit Range</Label>
+            <Stack horizontal tokens={{ childrenGap: 16 }}>
+                <Stack >
+                    <BasicListParameter
+                        placeholder="Earliest"
+                        handler={handler}
+                        performanceTrendOptions={performanceTrendOptions}
+                        params={startDropDownParams}
+                        selectedKey={selectedStartKey}
+                        request={EMPTY_REFRESH_REQUEST}
+                        onChange={(option: IDropdownOption<PerformanceTrendOptionData>) => {
+                            const key = option.key as string;
+                            performanceTrendOptions.getTransactionSession().toggleStartCommitRange(Number.parseInt(key));
+                        }}
+                        selectorWidth={150}
+                        calloutWidth={150}
+                    />
+                </Stack>
+
+                <Stack >
+                    <BasicListParameter
+                        placeholder="Latest"
+                        handler={handler}
+                        performanceTrendOptions={performanceTrendOptions}
+                        params={endDropDownParams}
+                        selectedKey={selectedEndKey}
+                        request={EMPTY_REFRESH_REQUEST}
+                        onChange={(option: IDropdownOption<PerformanceTrendOptionData>) => {
+                            const key = option.key as string;
+                            performanceTrendOptions.getTransactionSession().toggleEndCommitRange(Number.parseInt(key));
+                        }}
+                        selectorWidth={150}
+                        calloutWidth={150}
+                    />
+                </Stack>
+            </Stack>
+        </Stack>
+    );
+});
\ No newline at end of file
diff --git a/Engine/Source/Programs/Horde/HordeDashboard/plugins/performancetrends/components/DropdownComponent.tsx b/Engine/Source/Programs/Horde/HordeDashboard/plugins/performancetrends/components/DropdownComponent.tsx
new file mode 100644
--- /dev/null
+++ b/Engine/Source/Programs/Horde/HordeDashboard/plugins/performancetrends/components/DropdownComponent.tsx
@@ -0,0 +1,74 @@
+// Copyright Epic Games, Inc. All Rights Reserved.
+
+import { DirectionalHint, Dropdown, IDropdownOption } from "@fluentui/react";
+import { PerformanceTrendOptionsController } from "../filters/PerformanceTrendOptionsController";
+import { DataHandlerRefreshRequest, PerformanceTrendOptionsDataHandler } from "../filters/PerformanceTrendOptionsDataHandler";
+import { KEY_SEPARATOR } from "../responses/FilterKeys";
+
+// #region -- Performance Trend Dropdown Params --
+
+export const PARAMETER_KEY_PREFIX: string = "parameter_key";
+
+export type PerformanceTrendOptionData = {
+    id: string,
+    text: string,
+    group?: string,
+    tooltip?: string
+}
+
+export type PerformanceTrendOptionList = {
+    items: PerformanceTrendOptionData[]
+    label: string,
+    tooltip: string
+}
+
+// #endregion -- Performance Trend Dropdown Params --
+
+/**
+ * General purpose basic list dropdown.
+ * @param placeholder The placeholder string for the dropdown.
+ * @param handler The options data handler.
+ * @param performanceTrendOptions The options controller.
+ * @param params The list of performance trend options.
+ * @param disabled Whether the dropdown is disabled or not.
+ * @param selectedKey The currently selected key.
+ * @param request The data refresh request to use upon selection.
+ * @param onChange The on change delegate to invoke upon dropdown selection change.
+ * @todo This is a code-sharing candidate between Performance Trends and Build Health.
+ */
+export const BasicListParameter: React.FC<{ placeholder: string, handler: PerformanceTrendOptionsDataHandler, performanceTrendOptions: PerformanceTrendOptionsController, params: PerformanceTrendOptionList, disabled?: boolean, selectedKey?: string, request: DataHandlerRefreshRequest, onChange: (option: IDropdownOption<PerformanceTrendOptionData>) => void, selectorWidth?: number, calloutWidth?: number }> = function ConstructBasicListParamater({ placeholder, params: param, handler, performanceTrendOptions, selectedKey, request, onChange, selectorWidth = 200, calloutWidth = 300 }) {
+    const key = `${PARAMETER_KEY_PREFIX}${KEY_SEPARATOR}${param.label}`;
+    const doptions: IDropdownOption<PerformanceTrendOptionData>[] = [];
+
+    param.items.forEach((item) => {
+        doptions.push({
+            key: item.id,
+            text: item.text,
+        });
+    });
+
+    return (<Dropdown id={key}
+        key={key}
+        label={param.label}
+        placeholder={placeholder}
+        options={doptions}
+        selectedKey={selectedKey}
+        dropdownWidth={selectorWidth}
+        styles={{ dropdown: { width: calloutWidth } }}
+        calloutProps={{
+            directionalHint: DirectionalHint.rightTopEdge,
+            alignTargetEdge: true,
+        }}
+        onChange={(_ev, option) => {
+            let castedOption = option as IDropdownOption<PerformanceTrendOptionData>;
+            if (!castedOption) {
+                return;
+            }
+            onChange(castedOption as IDropdownOption<PerformanceTrendOptionData>);
+        }}
+        onDismiss={() => {
+            performanceTrendOptions.commitTransactionSession();
+            handler.requestHierarchicalRefresh(request, undefined, (_request: DataHandlerRefreshRequest) => { });
+        }}
+    />);
+};
\ No newline at end of file
diff --git a/Engine/Source/Programs/Horde/HordeDashboard/plugins/performancetrends/components/LineGraphComponent.tsx b/Engine/Source/Programs/Horde/HordeDashboard/plugins/performancetrends/components/LineGraphComponent.tsx
new file mode 100644
--- /dev/null
+++ b/Engine/Source/Programs/Horde/HordeDashboard/plugins/performancetrends/components/LineGraphComponent.tsx
@@ -0,0 +1,688 @@
+// Copyright Epic Games, Inc. All Rights Reserved.
+
+import React, { useEffect, useRef, useState } from "react";
+import * as d3 from "d3";
+import { getHordeStyling, preloadFonts } from "horde/styles/Styles";
+import { graphColors } from "hordePlugins/analytics/telemetryData";
+import { ISummaryMetric } from "../metrictypes/PerformanceTrendsTypes";
+
+export type MetricSelectionHandler<T extends ISummaryMetric> = (selectedSeriesElements: T[]) => void
+
+const LEGEND_SIZE_DEFAULT: number = 225;
+const DEFAULT_TOOLTIP_WIDTH_ESTIMATE = 200;
+
+/**
+ * Properties for the line graph component.
+ */
+export interface LineGraphProps<T extends ISummaryMetric> {
+    /**
+     * 2d array representing a list of series, each containing it's own numerical data.
+     */
+    series: SeriesConfig<T>[];
+
+    /**
+     * Defines an x accessor for the graph.
+     * @param item The item.
+     * @param index The index of the item.
+     * @returns The numerical representation along the x axis.
+     */
+    xAccessor?: (item: T, index: number) => number;
+
+    /**
+     * The y axis label.
+     */
+    yAxisLabel: string,
+
+    /**
+     * The x axis label.
+     */
+    xAxisLabel: string,
+
+    /**
+     * The maximum width of the graph.
+     */
+    width?: number;
+
+    /**
+     * The prescribed maximum height of the graph.
+     */
+    height?: number;
+
+    /**
+     * Whether to have the y axis start at zero, or whether to perform min-max clamping for the limits.
+     */
+    yAxisZeroScale?: boolean;
+
+    /**
+     * Whether to apply curve smoothing to the line graph.
+     */
+    applyCurveSmoothing?: boolean;
+
+    /**
+     * Onclick delegate
+     * @returns 
+     */
+    onClick?: MetricSelectionHandler<T>;
+}
+
+export interface SeriesConfig<T extends ISummaryMetric> {
+    data: T[];
+    yAccessor: (item: T) => number;
+    label: string;
+    color?: string;
+}
+
+/**
+ * LineGraph component used to draw a line graph using input data, labels, etc.
+ * @param props The properties to use in constructing the Line Graph.
+ * @returns The line graph component.
+ */
+export function LineGraph<T extends ISummaryMetric>({ series, xAccessor, yAxisLabel, xAxisLabel, width = 300, height = 150, yAxisZeroScale = false, applyCurveSmoothing = false, onClick = undefined }: LineGraphProps<T>): React.ReactElement {
+    const { modeColors, _hordeClasses, _detailClasses } = getHordeStyling();
+
+    // #region -- Graph Component Refs & State --
+
+    const svgRef = useRef<SVGSVGElement | null>(null);
+    const graphRef = useRef<d3.Selection<SVGGElement, unknown, null, undefined> | null>(null);
+    const tooltipRef = useRef<HTMLDivElement | null>(null);
+
+    const [selectedIndex, setSelectedIndex] = useState(-1);
+
+    // #endregion -- Graph Component Refs --
+
+    type LineGraphData = {
+        min: number;
+        max: number;
+        avg: number;
+        data: T[];
+    }
+
+    // #region -- Tool Tip Helpers --
+
+    const clearTableTooltipRows = (tooltipContent: HTMLDivElement) => {
+        const container = tooltipContent.querySelector(
+            ".tooltipContentContainer"
+        ) as HTMLDivElement | null;
+
+        if (!container) {
+            return;
+        }
+
+        // Fast + safe
+        container.replaceChildren();
+    };
+
+    let generateTableTooltipRow = (tooltipContent: HTMLDivElement, label: string, value: string, color: string) => {
+        // Find or create the parent container
+        let container = tooltipContent.querySelector(
+            ".tooltipContentContainer"
+        ) as HTMLDivElement | null;
+
+        if (!container) {
+            container = document.createElement("div");
+            container.className = "tooltipContentContainer";
+            tooltipContent.appendChild(container);
+        }
+
+        // Create the row
+        const row = document.createElement("div");
+        row.className = "tooltipRow";
+
+        const labelEl = document.createElement("strong");
+        labelEl.textContent = `${label}: `;
+
+        const valueEl = document.createElement("span");
+        valueEl.textContent = value;
+
+        if (color) {
+            row.style.color = color
+        }
+
+        row.appendChild(labelEl);
+        row.appendChild(valueEl);
+
+        container.appendChild(row);
+    }
+
+    // #endregion -- Tool Tip Helpers --
+
+    // #region -- Data Processing Helpers -- 
+
+    function summarizeSeriesData(summarizedSeries: Map<number, { min: number; max: number; avg: number; data: T[]; }>[], minYAxisValue: number, maxYAxisValue: number, maxXAxisValue: number, minXAxisValue: number) {
+        let count: number = 0;
+
+        for (let idx: number = 0; idx < series.length; ++idx) {
+            let seriesSet: Map<number, LineGraphData> = new Map<number, LineGraphData>;
+            summarizedSeries.push(seriesSet);
+
+            for (let idxY: number = 0; idxY < series[idx].data.length; ++idxY) {
+                let element: T = series[idx].data[idxY];
+                let resolvedXBucket = xAccessor ? xAccessor(series[idx].data[idxY], count) : count;
+
+                let value: number = series[idx].yAccessor(element);
+
+                minYAxisValue = minYAxisValue < value ? minYAxisValue : value;
+                maxYAxisValue = maxYAxisValue > value ? maxYAxisValue : value;
+
+                if (!seriesSet.has(resolvedXBucket)) {
+                    let datum: LineGraphData = { min: value, max: value, data: [element], avg: value };
+                    seriesSet.set(resolvedXBucket, datum);
+                    maxXAxisValue = maxXAxisValue > resolvedXBucket ? maxXAxisValue : resolvedXBucket;
+                    minXAxisValue = minXAxisValue < resolvedXBucket ? minXAxisValue : resolvedXBucket;
+                } else {
+                    let datum: LineGraphData = seriesSet.get(resolvedXBucket)!;
+
+                    datum.max = datum.max > value ? datum.max : value;
+                    datum.min = datum.min < value ? datum.min : value;
+
+                    datum.data.push(element);
+                    datum.avg = datum.avg + (value - datum.avg) / datum.data.length;
+                }
+
+                count += 1;
+            }
+        }
+        return { minYAxisValue, maxYAxisValue, maxXAxisValue, minXAxisValue };
+    }
+
+    // Find closest bucket to hovered X
+    const findClosestBucket = (seriesMap: Map<number, LineGraphData>, xLocation: number): [number, LineGraphData] | null => {
+        const buckets = Array.from(seriesMap.keys()).sort((a, b) => a - b);
+
+        if (buckets.length === 0) {
+            return null;
+        }
+
+        let closest = buckets[0];
+        let closestDist = Math.abs(xLocation - closest);
+
+        for (const bucket of buckets) {
+            const dist = Math.abs(xLocation - bucket);
+            if (dist < closestDist) {
+                closest = bucket;
+                closestDist = dist;
+            }
+        }
+
+        return seriesMap.has(closest) ? [closest, seriesMap.get(closest)!] : null;
+    };
+
+    // #endregion -- Data Processing Helpers --
+
+    // #region -- Use Effects --
+
+    useEffect(() => {
+        if (!svgRef.current) return;
+
+        let summarizedSeries: Map<number, LineGraphData>[] = [];
+        let maxYAxisValue: number = Number.NEGATIVE_INFINITY;
+        let minYAxisValue: number = Number.MAX_VALUE;
+        let maxXAxisValue: number = Number.NEGATIVE_INFINITY;
+        let minXAxisValue: number = Number.MAX_VALUE;
+
+        ({ minYAxisValue, maxYAxisValue, maxXAxisValue, minXAxisValue } = summarizeSeriesData(summarizedSeries, minYAxisValue, maxYAxisValue, maxXAxisValue, minXAxisValue));
+
+        let colorLabels: string[] = series.map((x) => x.color ?? null).filter((value, idx) => value !== null);
+
+        // #region -- Graph Dimensions & Measures --
+        // SVG width is already reduced by legend width in the render
+        const svgWidth = width - LEGEND_SIZE_DEFAULT;
+
+        // Dimensions
+        const margin = { top: 10, right: 10, bottom: 45, left: 55 };
+        const innerWidth = svgWidth - margin.left - margin.right;
+        const innerHeight = height - margin.top - margin.bottom;
+
+        // Compute min, max, midpoint for centered line
+        const mid = (maxYAxisValue + minYAxisValue) / 2;
+        const halfRange = (maxYAxisValue - minYAxisValue) / 2;
+        const lowerBoundYAxis = yAxisZeroScale ? 0 : mid - halfRange;
+        const upperBoundYAxis = yAxisZeroScale ? (maxYAxisValue) * 1.25 : mid + halfRange;
+
+        // Scales
+        // Collect all unique x buckets across all series
+        const allBucketsSet = new Set<number>();
+        for (const seriesMap of summarizedSeries) {
+            for (const bucket of seriesMap.keys()) {
+                allBucketsSet.add(bucket);
+            }
+        }
+        const sortedBuckets = Array.from(allBucketsSet).sort((a, b) => a - b);
+
+        // When xAccessor is provided, use scalePoint for evenly-spaced discrete values (changelists)
+        // Otherwise, use scaleLinear for proportional numeric spacing
+        let xScale: (value: number) => number;
+        let xInvert: (pixelX: number) => number;
+
+        if (xAccessor) {
+            const xPoint = d3.scalePoint<number>()
+                .domain(sortedBuckets)
+                .range([0, innerWidth]);
+
+            xScale = (value: number) => xPoint(value) ?? 0;
+
+            xInvert = (pixelX: number): number => {
+                if (sortedBuckets.length === 0) return 0;
+                if (sortedBuckets.length === 1) return sortedBuckets[0];
+
+                let closestBucket = sortedBuckets[0];
+                let closestDist = Math.abs(pixelX - (xPoint(closestBucket) ?? 0));
+
+                for (const bucket of sortedBuckets) {
+                    const bucketPixel = xPoint(bucket) ?? 0;
+                    const dist = Math.abs(pixelX - bucketPixel);
+                    if (dist < closestDist) {
+                        closestBucket = bucket;
+                        closestDist = dist;
+                    }
+                }
+                return closestBucket;
+            };
+        } else {
+            const xLinear = d3.scaleLinear()
+                .domain([minXAxisValue, maxXAxisValue])
+                .range([0, innerWidth]);
+
+            xScale = (value: number) => xLinear(value);
+            xInvert = (pixelX: number) => Math.round(xLinear.invert(pixelX));
+        }
+
+        const y = d3.scaleLinear()
+            .domain([upperBoundYAxis, lowerBoundYAxis])
+            .range([0, innerHeight]);
+
+        // #endregion -- Graph Dimensions & Measures --
+
+        // #region -- Core Graph Setup --
+
+        // Clear previous contents
+        const svg = d3.select(svgRef.current);
+        svg.selectAll("*").remove();
+
+        // Main graph group
+        const g = svg.append("g")
+            .attr("transform", `translate(${margin.left},${margin.top})`);
+
+        graphRef.current = g;
+
+        // #endregion -- Core Graph Setup --
+
+        // #region -- Axis --
+
+        // Axes
+        if (xAccessor) {
+            // For scalePoint: use sparse ticks if there are many buckets; we cap at 10 for now
+            const xPoint = d3.scalePoint<number>()
+                .domain(sortedBuckets)
+                .range([0, innerWidth]);
+
+            const maxTicks = 10;
+            const step = Math.ceil(sortedBuckets.length / maxTicks);
+            const sparseBuckets = sortedBuckets.filter((_, i) => i % step === 0);
+
+            g.append("g")
+                .attr("transform", `translate(0,${innerHeight})`)
+                .call(d3.axisBottom(xPoint).tickValues(sparseBuckets).tickFormat(d => d.toString()))
+                .style("font-family", preloadFonts[0]);
+        } else {
+            // For scaleLinear: use default ticks
+            const xLinear = d3.scaleLinear()
+                .domain([minXAxisValue, maxXAxisValue])
+                .range([0, innerWidth]);
+
+            g.append("g")
+                .attr("transform", `translate(0,${innerHeight})`)
+                .call(d3.axisBottom(xLinear).ticks(5))
+                .style("font-family", preloadFonts[0]);
+        }
+
+        g.append("g")
+            .call(d3.axisLeft(y).ticks(5))
+            .style("font-family", preloadFonts[0]);
+
+        // Add Axis Labels
+        // X axis label
+        g.append("text")
+            .attr("x", innerWidth / 2)
+            .attr("y", innerHeight + margin.bottom - 4)
+            .attr("text-anchor", "middle")
+            .attr("fill", modeColors.text)
+            .style("font-family", preloadFonts[0])
+            .style("font-size", "12px")
+            .text(xAxisLabel);
+
+        // Y axis label
+        g.append("text")
+            .attr("transform", `rotate(-90)`)
+            .attr("x", -innerHeight / 2)
+            .attr("y", -margin.left + 12)
+            .attr("text-anchor", "middle")
+            .attr("fill", modeColors.text)
+            .style("font-family", preloadFonts[0])
+            .style("font-size", "12px")
+            .text(yAxisLabel);
+
+        // #endregion -- Axis --
+
+        // #region -- Line Drawing --
+
+        // Line generator
+        // Line generator for bucketed data
+        const lineGenerator = d3.line<[number, LineGraphData]>()
+            .x(([xBucket, _]) => xScale(xBucket))
+            .y(([_, data]) => y(data.avg));  // or data.min, data.max
+
+        if (applyCurveSmoothing) {
+            lineGenerator.curve(d3.curveMonotoneX);
+        }
+
+        // if anyone didn't specify colour, just default to graph colours for now.
+        const colors = colorLabels.length === series.length ? colorLabels : graphColors;
+
+        // Area generator for min/max range
+        const areaGenerator = d3.area<[number, LineGraphData]>()
+            .x(([xBucket, _]) => xScale(xBucket))
+            .y0(([_, data]) => y(data.min))
+            .y1(([_, data]) => y(data.max));
+
+        if (applyCurveSmoothing) {
+            areaGenerator.curve(d3.curveMonotoneX);
+        }
+
+        // Draw range bands (before lines so lines are on top)
+        for (let idx: number = 0; idx < summarizedSeries.length; ++idx) {
+            const sortedEntries = Array.from(summarizedSeries[idx].entries())
+                .sort(([a], [b]) => a - b);
+
+            g.append("path")
+                .datum(sortedEntries)
+                .attr("fill", colors[idx])
+                .attr("fill-opacity", 0.15)
+                .attr("stroke", "none")
+                .attr("d", areaGenerator)
+                .attr("class", `series-range series-${idx}`);
+        }
+
+        // Draw lines from bucketed data
+        for (let idx: number = 0; idx < summarizedSeries.length; ++idx) {
+            const seriesMap = summarizedSeries[idx];
+            const sortedEntries = Array.from(seriesMap.entries())
+                .sort(([a], [b]) => a - b);
+
+            g.append("path")
+                .datum(sortedEntries)
+                .attr("fill", "none")
+                .attr("stroke", colors[idx])
+                .attr("stroke-width", 2)
+                .attr("d", lineGenerator)
+                .attr("class", `series-line series-${idx}`);
+        }
+
+        // #endregion -- Line Drawing --
+
+        // #region -- Tooltip & Accessibility Controls --
+
+        // Vertical hover line
+        const hoverLine = g.append("line")
+            .attr("stroke", "gray")
+            .attr("stroke-width", 1)
+            .attr("y1", 0)
+            .attr("y2", innerHeight)
+            .style("opacity", 0);
+
+        // Tooltip styling
+        const tooltipElement = tooltipRef.current;
+
+        if (tooltipElement) {
+            tooltipElement.style.position = "fixed";
+            tooltipElement.style.pointerEvents = "none";
+            tooltipElement.style.background = modeColors.background;
+            tooltipElement.style.border = "1px solid #aaa";
+            tooltipElement.style.padding = "4px 6px";
+            tooltipElement.style.borderRadius = "4px";
+            tooltipElement.style.fontFamily = preloadFonts[0];
+            tooltipElement.style.fontSize = "12px";
+            tooltipElement.style.color = modeColors.text;
+            tooltipElement.style.opacity = "0";
+            tooltipElement.style.zIndex = "9999";
+        }
+
+        // #endregion -- Tooltip & Accessibility Controls --
+
+        const highlightLayer = g.append("g")
+            .attr("class", "tooltip-highlights");
+
+        // Overlay rectangle to capture mouse
+        g.append("rect")
+            .attr("width", innerWidth)
+            .attr("height", innerHeight)
+            .style("fill", "none")
+            .style("pointer-events", "all")
+            .on("mousemove", (event: MouseEvent) => {
+                if (!tooltipElement) {
+                    return;
+                }
+
+                const parentRect = tooltipElement.parentElement!.getBoundingClientRect();
+                const mouseX = event.clientX - parentRect.left - margin.left;
+                const mouseY = event.clientY - parentRect.top - margin.top;
+                const hoveredX = xInvert(mouseX);
+
+                // Check bounds
+                if (hoveredX < minXAxisValue || hoveredX > maxXAxisValue) {
+                    return;
+                }
+
+                tooltipElement.style.opacity = "1";
+
+                // Flip tooltip to left side if it would overflow the viewport
+                const tooltipWidth = tooltipElement.offsetWidth || DEFAULT_TOOLTIP_WIDTH_ESTIMATE; // estimate if not yet rendered
+                const wouldOverflow = event.clientX + 15 + tooltipWidth > window.innerWidth;
+
+                if (wouldOverflow) {
+                    tooltipElement.style.left = `${event.clientX - tooltipWidth - 15}px`;
+                } else {
+                    tooltipElement.style.left = `${event.clientX + 15}px`;
+                }
+                tooltipElement.style.top = `${event.clientY + 10}px`;
+
+                clearTableTooltipRows(tooltipElement);
+                highlightLayer.selectAll("*").remove();
+
+                let snappedX: number | null = null;
+
+                for (let idx: number = 0; idx < summarizedSeries.length; ++idx) {
+                    const result = findClosestBucket(summarizedSeries[idx], hoveredX);
+                    if (!result) continue;
+
+                    const [bucket, data] = result;
+
+                    const label = legendLabels[idx];
+
+                    // Only show this series if it matches the snapped bucket
+                    if (hoveredX !== bucket) {
+                        continue;
+                    }
+
+                    // Use first series' bucket for vertical line
+                    if (snappedX === null) {
+                        snappedX = bucket;
+                        generateTableTooltipRow(tooltipElement, `(${bucket})`, "", modeColors.messageText);
+                    }
+
+                    // Tooltip: show avg, min, max, count
+                    const tooltipText = `avg: ${data.avg.toFixed(2)} | min: ${data.min.toFixed(2)} | max: ${data.max.toFixed(2)} | n: ${data.data.length}`;
+                    generateTableTooltipRow(tooltipElement, label, tooltipText, colors[idx]);
+
+                    // Highlight circle at avg
+                    const xpos = xScale(bucket);
+                    const ypos = y(data.avg);
+
+                    highlightLayer.append("circle")
+                        .attr("cx", xpos)
+                        .attr("cy", ypos)
+                        .attr("r", 5)
+                        .attr("fill", colors[idx]);
+
+                    // Draw a line across the range
+                    highlightLayer.append("line")
+                        .attr("x1", xpos)
+                        .attr("x2", xpos)
+                        .attr("y1", y(data.min))
+                        .attr("y2", y(data.max))
+                        .attr("stroke", colors[idx])
+                        .attr("stroke-width", 2)
+                        .attr("opacity", 0.5);
+
+                    // Annotate the range with smaller dots
+                    highlightLayer.append("circle")
+                        .attr("cx", xpos)
+                        .attr("cy", y(data.min))
+                        .attr("r", 2)
+                        .attr("fill", colors[idx]);
+
+                    highlightLayer.append("circle")
+                        .attr("cx", xpos)
+                        .attr("cy", y(data.max))
+                        .attr("r", 2)
+                        .attr("fill", colors[idx]);
+                }
+
+                // Show vertical line at snapped bucket
+                if (snappedX !== null) {
+                    const snappedPixel = xScale(snappedX);
+                    hoverLine
+                        .attr("x1", snappedPixel)
+                        .attr("x2", snappedPixel)
+                        .style("opacity", 1);
+                }
+            })
+            .on("mouseleave", () => {
+                if (!tooltipElement) {
+                    return;
+                }
+                highlightLayer.selectAll("*").remove();
+
+                tooltipElement.style.opacity = "0";
+                hoverLine.style("opacity", 0);
+            })
+            .on("click", (event: MouseEvent) => {
+                if (!tooltipElement || !onClick) {
+                    return;
+                }
+
+                const parentRect = tooltipElement.parentElement!.getBoundingClientRect();
+                const mouseX = event.clientX - parentRect.left - margin.left;
+                const index = xInvert(mouseX);
+
+                let selectedItems: T[] = [];
+                for (let idx: number = 0; idx < summarizedSeries.length; ++idx) {
+                    const result = findClosestBucket(summarizedSeries[idx], index);
+
+                    if (!result) {
+                        continue;
+                    }
+
+                    const [bucket, data] = result;
+                    if (bucket === index) {
+                        selectedItems.push(...data.data);
+                    }
+                }
+
+                onClick(selectedItems);
+            });
+
+        // #endregion -- Legend --
+
+    }, [series, width, height, yAxisZeroScale, applyCurveSmoothing]);
+
+    useEffect(() => {
+        if (!graphRef.current) {
+            return;
+        }
+
+        const lines = graphRef.current.selectAll<SVGPathElement, unknown>(".series-line");
+        const ranges = graphRef.current.selectAll<SVGPathElement, unknown>(".series-range");
+
+        if (selectedIndex === -1) {
+            lines.style("opacity", 1);
+            ranges.style("opacity", 0.65);
+        } else {
+            lines.style("opacity", 0.2);
+            ranges.style("opacity", 0.05);
+
+            graphRef.current.select<SVGPathElement>(`.series-line.series-${selectedIndex}`).style("opacity", 1);
+            graphRef.current.select<SVGPathElement>(`.series-range.series-${selectedIndex}`).style("opacity", 0.65);
+        }
+    }, [selectedIndex]);
+
+    // #endregion -- Use Effects --
+
+    // Compute colors for legend (same logic as in useEffect)
+    const legendColors = series.map((s, i) => s.color ?? graphColors[i % graphColors.length]);
+    const legendLabels = series.map((s) => s.label);
+    const maxLegendHeight = Math.max(height - 40, 100); // Leave some margin, min 100px
+
+    return (
+        <div style={{ position: "relative", display: "flex" }}>
+            <svg ref={svgRef} width={width - LEGEND_SIZE_DEFAULT} height={height}></svg>
+            <div
+                style={{
+                    width: LEGEND_SIZE_DEFAULT,
+                    maxHeight: maxLegendHeight,
+                    overflowY: "auto",
+                    overflowX: "hidden",
+                    background: modeColors.content,
+                    borderRadius: 4,
+                    padding: 6,
+                    boxSizing: "border-box",
+                    fontFamily: preloadFonts[0],
+                    fontSize: 10,
+                    color: modeColors.text,
+                    marginTop: 20,
+                }}
+            >
+                {series.map((s, i) => (
+                    <div
+                        key={i}
+                        style={{
+                            display: "flex",
+                            alignItems: "flex-start",
+                            marginBottom: 6,
+                            cursor: "pointer",
+                            opacity: selectedIndex === -1 || selectedIndex === i ? 1 : 0.2,
+                        }}
+                        onClick={() => setSelectedIndex(prev => prev === i ? -1 : i)}
+                    >
+                        <div
+                            style={{
+                                width: 20,
+                                height: 2,
+                                backgroundColor: legendColors[i],
+                                marginTop: 6,
+                                marginRight: 6,
+                                flexShrink: 0,
+                            }}
+                        />
+                        <div style={{ flex: 1, minWidth: 0 }}>
+                            {legendLabels[i].split('\n').map((line, lineIdx) => (
+                                <div
+                                    key={lineIdx}
+                                    style={{
+                                        whiteSpace: "nowrap",
+                                        overflow: "hidden",
+                                        textOverflow: "ellipsis",
+                                    }}
+                                    title={line}
+                                >
+                                    {line}
+                                </div>
+                            ))}
+                        </div>
+                    </div>
+                ))}
+            </div>
+            <div ref={tooltipRef}></div>
+        </div>
+    );
+};
\ No newline at end of file
diff --git a/Engine/Source/Programs/Horde/HordeDashboard/plugins/performancetrends/components/MetricSummaryTypeDropdownSingle.tsx b/Engine/Source/Programs/Horde/HordeDashboard/plugins/performancetrends/components/MetricSummaryTypeDropdownSingle.tsx
new file mode 100644
--- /dev/null
+++ b/Engine/Source/Programs/Horde/HordeDashboard/plugins/performancetrends/components/MetricSummaryTypeDropdownSingle.tsx
@@ -0,0 +1,44 @@
+// Copyright Epic Games, Inc. All Rights Reserved.
+
+import { Stack, Label, IDropdownOption } from "@fluentui/react";
+import { observer } from "mobx-react-lite";
+import { PerformanceTrendOptionsController } from "../filters/PerformanceTrendOptionsController";
+import { DataHandlerRefreshRequest, expandRefreshRequest, PerformanceTrendOptionsDataHandler } from "../filters/PerformanceTrendOptionsDataHandler";
+import { PerformanceTrendOptionList, PerformanceTrendOptionData, BasicListParameter } from "./DropdownComponent";
+import { GetTestProjectResponse } from "../responses/GetTestProjectResponse";
+import { encodeMetricTypeKey } from "../responses/FilterKeys";
+
+/**
+ * Component used to model all known metric summary types (for the currently selected test project, and test identity), as a single select dropdown.
+ * @return The react component.
+ */
+export const MetricSummaryTypeDropdownSingle: React.FC<{ handler: PerformanceTrendOptionsDataHandler, performanceTrendOptions: PerformanceTrendOptionsController }> = observer(function ConstructProjectDropdownSingle({ handler, performanceTrendOptions }) {
+    const metricSummaryDropdownParams: PerformanceTrendOptionList = {
+        label: "",
+        items: [
+        ],
+        tooltip: "Select an option from the dropdown."
+    };
+
+    metricSummaryDropdownParams.items = handler.metricSummaryTypes.map((metricTypeParentProject: GetTestProjectResponse) => {
+        let listParam: PerformanceTrendOptionData = {
+            id: encodeMetricTypeKey(metricTypeParentProject),
+            text: metricTypeParentProject.summaryType,
+        };
+
+        return listParam;
+    }).sort((a, b) => a.text.localeCompare(b.text));
+
+    let selectedSummaryMetricTypeKey = Object.keys(performanceTrendOptions.state.enabledMetricSummaryTypes)[0] ?? null;
+    let request: DataHandlerRefreshRequest = expandRefreshRequest({ testTypes: true });
+    
+    return metricSummaryDropdownParams.items.length > 0 && (
+        <Stack>
+            <Label>Summary Type</Label>
+            <BasicListParameter placeholder={"Select a metric summary type"} handler={handler} performanceTrendOptions={performanceTrendOptions} params={metricSummaryDropdownParams} selectedKey={selectedSummaryMetricTypeKey} request={request} onChange={(option: IDropdownOption<PerformanceTrendOptionData>) => {
+                const key = option.key as string;
+                performanceTrendOptions.getTransactionSession().toggleSingleMetricSummaryType(key, option.text);
+            }} />
+        </Stack>
+    );
+});
\ No newline at end of file
diff --git a/Engine/Source/Programs/Horde/HordeDashboard/plugins/performancetrends/components/MultilistDropdownComponent.tsx b/Engine/Source/Programs/Horde/HordeDashboard/plugins/performancetrends/components/MultilistDropdownComponent.tsx
new file mode 100644
--- /dev/null
+++ b/Engine/Source/Programs/Horde/HordeDashboard/plugins/performancetrends/components/MultilistDropdownComponent.tsx
@@ -0,0 +1,211 @@
+import { DefaultButton, DropdownMenuItemType, IDropdownOption, ISelectableOption, mergeStyleSets, SelectableOptionMenuItemType, Stack } from "@fluentui/react";
+import { DataHandlerRefreshRequest, PerformanceTrendOptionsDataHandler } from "../filters/PerformanceTrendOptionsDataHandler";
+import { PerformanceTrendOptionsController } from "../filters/PerformanceTrendOptionsController";
+import { PARAMETER_KEY_PREFIX, PerformanceTrendOptionData, PerformanceTrendOptionList } from "./DropdownComponent";
+import { useCallback, useEffect, useLayoutEffect, useMemo, useRef, useState } from "react";
+
+// We should generalize this
+import { SearchableDropdown } from "horde/components/buildhealth/SearchableDropdownComponent";
+import { PerformanceTrendOptionsStateDiff } from "../filters/PerformanceTrendOptionsWriter";
+import { KEY_SEPARATOR } from "../responses/FilterKeys";
+
+const SELECT_CLEAR_ALL_SPECIAL_GROUP: string = "select_clear_all";
+
+const dropdownOptions = mergeStyleSets({
+    base: {
+        padding: "4px 8px",
+    }
+});
+
+/**
+ * React Component that is a multi select drop down.
+ * @returns React Component.
+ */
+export const MultiListParameter: React.FC<{ handler: PerformanceTrendOptionsDataHandler, performanceTrendOptions: PerformanceTrendOptionsController, params: PerformanceTrendOptionList, disabled?: boolean, selectedKeys: string[], request: DataHandlerRefreshRequest, onChange: (option: IDropdownOption<PerformanceTrendOptionData>) => void, enabledSelectAll?: boolean, onSelectAll?: (groupOption: IDropdownOption<PerformanceTrendOptionData>, optionSelectAllContext?: IDropdownOption<PerformanceTrendOptionData>[]) => void, onHierarchicalRefreshCompleted?: (receipts?: PerformanceTrendOptionsStateDiff) => void }> = function ConstructMultilistParamter({ params, handler, performanceTrendOptions, disabled, selectedKeys, request, onChange, enabledSelectAll = false, onSelectAll, onHierarchicalRefreshCompleted }) {
+    const [filterText, setFilterText] = useState("");
+    const changeVersion = performanceTrendOptions.optionsChangeVersion;
+    const key = `${PARAMETER_KEY_PREFIX}${KEY_SEPARATOR}${params.label}`;
+    const selectedSet = useMemo(() => new Set(selectedKeys), [selectedKeys]);
+    const groupRefs = useRef<(HTMLDivElement | null)[]>([]);
+
+    useEffect(() => { groupRefs.current = []; }, [params.items]);
+
+    // #region -- Data Memos --
+
+    const doptions = useMemo(() => {
+        const groupMap = new Map<string, PerformanceTrendOptionData[]>();
+
+        for (const item of params.items) {
+            const group = item.group ?? "__nogroup";
+            if (!groupMap.has(group)) groupMap.set(group, []);
+            groupMap.get(group)!.push(item);
+        }
+
+        const groups = Array.from(groupMap.keys()).sort();
+        const options: IDropdownOption<PerformanceTrendOptionData>[] = [];
+
+        for (const group of groups) {
+            if (group !== "__nogroup") {
+                options.push({ key: `group_${group}`, text: group, itemType: DropdownMenuItemType.Header });
+                if (enabledSelectAll) {
+                    options.push({
+                        key: `group_select_all${KEY_SEPARATOR}${group}`,
+                        itemType: DropdownMenuItemType.Header,
+                        text: SELECT_CLEAR_ALL_SPECIAL_GROUP,
+                    });
+                }
+            }
+
+            for (const item of groupMap.get(group)!) {
+                options.push({
+                    key: item.id,
+                    text: item.text,
+                    data: item,
+                    selected: selectedSet.has(item.id),
+                });
+            }
+        }
+
+        return options;
+    }, [params.items, selectedSet, enabledSelectAll]);
+
+    // #endregion -- Data Memos --
+    const filteredOptions = useMemo(() => {
+        const lowerFilter = filterText.toLowerCase();
+        return doptions.filter(o =>
+            o.itemType === DropdownMenuItemType.Header || !filterText
+                ? true
+                : o.text.toLowerCase().includes(lowerFilter)
+        );
+    }, [doptions, filterText]);
+
+    const filteredOptionsRef = useRef<IDropdownOption<PerformanceTrendOptionData>[]>([]);
+
+    useLayoutEffect(() => {
+        filteredOptionsRef.current = filteredOptions;
+    }, [filteredOptions]);
+
+    // #region -- Dropdown Callbacks --
+
+    const onRenderOption = useCallback(
+        (option?: ISelectableOption) => {
+            if (!option) {
+                return null;
+            }
+
+            // Select & Clear all for the given filtered list.
+            if (option.itemType === DropdownMenuItemType.Header && option.text === SELECT_CLEAR_ALL_SPECIAL_GROUP) {
+                let filteredOptionList = filteredOptionsRef.current.length === doptions.length ? undefined : filteredOptionsRef.current.filter(x => x.itemType === undefined || x.itemType !== SelectableOptionMenuItemType.Header);
+                let selectAllLabel = (filteredOptionList !== undefined && filteredOptionList.length > 0) ? "Select Filtered" : "Select All";
+                let showButton = filteredOptionList === undefined || filteredOptionList.length > 0;
+                return (
+                    <Stack horizontal tokens={{ childrenGap: 8 }} styles={{
+                        root:
+                        {
+                            width: "100%",
+                            height: "100%",
+                            alignItems: "center",
+                            justifyContent: "center",
+                        }
+                    }}>
+                        {/* We have opted to use buttons here as we: 1. don't want to persist the state of a select all option; 2. the semantics of storing a "Select/Clear All" value gets complicated when the action is applied to a filtered set; */}
+                        {showButton &&
+                            <DefaultButton title="Selects all options (or currently filtered) within the group." styles={{ root: { flex: 1, height: 20 } }} text={selectAllLabel}
+                                onClick={() => {
+                                    if (!option) return;
+
+                                    const synthesizedOption: IDropdownOption<PerformanceTrendOptionData> = {
+                                        ...option,
+                                        selected: true,
+                                    };
+
+                                    // Only send a suggested list if we actually have filtered something.
+                                    let filteredOptionList = filteredOptionsRef.current.length === doptions.length ? undefined : filteredOptionsRef.current;
+                                    option && onSelectAll?.(synthesizedOption, filteredOptionList)
+                                }}
+                            />
+                        }
+                        <DefaultButton title="Clears all options within the group." styles={{ root: { flex: 1, height: 20 } }} text="Clear All"
+                            onClick={() => {
+                                if (!option) return;
+
+                                const synthesizedOption: IDropdownOption<PerformanceTrendOptionData> = {
+                                    ...option,
+                                    selected: false,
+                                };
+
+                                // We do not send a suggested list for clear all; we just clear all.
+                                option && onSelectAll?.(synthesizedOption, undefined)
+                            }} />
+                    </Stack>
+                );
+            }
+            if (option.itemType === DropdownMenuItemType.Header) {
+                return <div ref={el => groupRefs.current[option.key as string] = el} style={{ fontWeight: "bold", padding: "4px 8px" }}>{option.text}</div>;
+            }
+
+            return (
+                // Due to some load testing, when we get 1000s of entires we don't want to bloat the load time with ToolTipHost construction.
+                <div title={`Group: ${option.data?.group || ""}\n${option.text}\n${option.data?.tooltip}`} className={dropdownOptions.base}>{option.text}</div>
+            );
+        }, []);
+
+    const handleDropdownChange = useCallback((_ev, option) => {
+        if (option && option.itemType !== DropdownMenuItemType.Header) {
+            onChange(option);
+        }
+
+    }, [onChange]);
+
+    const handleMenuDismiss = useCallback(() => {
+        setFilterText("");
+        let buildHealthOptionStateReceipts = performanceTrendOptions.commitTransactionSession();
+        handler.requestHierarchicalRefresh(request, undefined,
+            // Upon completion of the hierarchical refresh, we want to commit the transaction session for the options. 
+            // We do this because the data handler can populate the template & step cache which may be useful for any downstream receipt processing, and actions that depend on the side-effects of updated hierarchy data.
+            (_request: DataHandlerRefreshRequest) => {
+                if (onHierarchicalRefreshCompleted) {
+                    onHierarchicalRefreshCompleted(buildHealthOptionStateReceipts);
+                }
+            });
+
+    }, [request, performanceTrendOptions]);
+
+    // #endregion -- Dropdown Callbacks --
+
+    return (
+        <SearchableDropdown
+            id={changeVersion.toString() + key}
+            key={key}
+            placeholder={params.label}
+            options={filteredOptions}
+            multiSelect
+            groupRefs={groupRefs}
+            selectedKeys={selectedKeys}
+            onChange={handleDropdownChange}
+            onDismiss={handleMenuDismiss}
+            onSearchValueChanged={(newValue) => {
+                // This stems from a very awkward bug where we need to force the underlying dropdown to resynchronize 
+                // on deletions by sending a empty filter - forcing the whole collection to be re-populated.
+                // We also need to force a complete animation frame in order to propagate this through, resulting in the
+                // underlying Dropdown component internal index synchronization.
+                // Without doing so, upon "deletion" changes the onChange will receive the wrong DropdownOption.
+                if (newValue.length < filterText.length) {
+                    setFilterText("");
+                    requestAnimationFrame(() => {
+                        setFilterText(newValue);
+                    });
+                } else {
+                    setFilterText(newValue);
+                }
+            }}
+            onRenderOption={onRenderOption}
+            styles={{
+                dropdown: {
+                    width: 300, maxWidth: 900,
+                },
+            }}
+            disabled={disabled}
+        />
+    );
+};
\ No newline at end of file
diff --git a/Engine/Source/Programs/Horde/HordeDashboard/plugins/performancetrends/components/PerformanceTrendUrl.tsx b/Engine/Source/Programs/Horde/HordeDashboard/plugins/performancetrends/components/PerformanceTrendUrl.tsx
new file mode 100644
--- /dev/null
+++ b/Engine/Source/Programs/Horde/HordeDashboard/plugins/performancetrends/components/PerformanceTrendUrl.tsx
@@ -0,0 +1,324 @@
+// Copyright Epic Games, Inc. All Rights Reserved.
+
+import { useEffect } from "react";
+import { PerformanceTrendOptionsController } from "../filters/PerformanceTrendOptionsController";
+import { PerformanceTrendOptionsDataHandler } from "../filters/PerformanceTrendOptionsDataHandler";
+import { useLocation, useNavigate } from "react-router";
+import { PerformanceTrendOptionsState } from "../filters/PerformanceTrendOptionsState";
+import { GroupByOption, PerformanceTrendUIOptionsState } from "../filters/PerformanceTrendUIOptionsState";
+import { IPerformanceTrendOptionWriter } from "../filters/PerformanceTrendOptionsWriter";
+import { PerformanceTrendQueryParams, PerformanceTrendUIParams } from "../filters/PerformanceTrendUrlParams";
+
+import LZString from "lz-string"
+import { runInAction } from "mobx";
+import { observer } from "mobx-react-lite";
+import { DecodedPlatformKey, DecodedStreamKey, DecodedTestIdentityKey, DecodedTestTypeKey, decodeMetricTypeKey, decodePlatformKey, decodeStreamKey, decodeTestIdentity, decodeTestTypeKey, decodeViewablePropertyKey } from "../responses/FilterKeys";
+
+// #region -- Consts --
+
+// #region -- Options --
+
+export const QUERY_SCHEMA_PARAM_NAME: string = "query_schema";
+export const PROJECTS_URL_PARAM_NAME: string = "projects";
+export const TEST_IDENTITIES_URL_PARAM_NAME: string = "testIdentities";
+export const TEST_TYPES_URL_PARAM_NAME: string = "testTypes";
+export const METRIC_SUMMARIES_URL_PARAM_NAME: string = "metricSummaries";
+export const STREAMS_URL_PARAM_NAME: string = "streams";
+export const PLATFORM_URL_PARAM_NAME: string = "platforms";
+export const START_COMMIT_PARAM_NAME: string = "start_commit";
+export const END_COMMIT_PARAM_NAME: string = "end_commit";
+
+// #endregion -- Options --
+
+// #region -- UI Options --
+
+export const UI_CURVE_SMOOTHING_PARAM_NAME: string = "lineSmoothing";
+export const UI_GROUP_BY_PARAM_NAME: string = "groupBy";
+export const UI_VIEWABLE_PROPERTIES_PARAM_NAME: string = "viewProps";
+
+// #endregion -- UI Options --
+
+// #endregion -- Consts --
+
+/**
+ * React Component used to manage the URL query params, and refresh it based on build health options. This is separate from the main component to scope rerenders to solely this component.
+ * @returns React Component.
+ */
+export const PerformanceTrendUrlSync: React.FC<{ buildHealthController: PerformanceTrendOptionsController; handler: PerformanceTrendOptionsDataHandler; }> = observer(function PerformanceTrendUrlSync({ buildHealthController, handler }) {
+    const location = useLocation();
+    const navigate = useNavigate();
+
+    // We only ever receive from location search once.
+    useEffect(() => {
+        const initializeFromUrl = async () => {
+            let searchParams: URLSearchParams = new URLSearchParams(location.search)
+            const params = parseQueryParams(searchParams);
+            const upgradedParams = upgradeQueryParams(params.performanceTrendQueryParams, buildHealthController.querySchemaVersion);
+
+            await loadOptionsFromParams(upgradedParams, params.performanceTrendUIParams, buildHealthController, handler);
+        };
+
+        initializeFromUrl();
+    }, []);
+
+    useEffect(() => {
+        const desiredSearch = toNavigationQuery(buildHealthController.state, buildHealthController.uiState, buildHealthController.querySchemaVersion, handler).toString();
+        if (desiredSearch && location.search !== `?${desiredSearch}`) {
+            navigate(`${location.pathname}?${desiredSearch}`, { replace: true });
+        }
+    }, [navigate, buildHealthController.optionsChangeVersion, buildHealthController.uiOptionsChangeVersion, location.pathname, location.search]);
+
+    return null;
+});
+
+// #region -- Query Serialization & Deserialization --
+
+/**
+ * Parses the query params from the search url.
+ * @param searchParams The search params to parse.
+ * @returns The resulting query params.
+ * @todo This is a code-sharing candidate between Performance Trends and Build Health.
+ */
+export function parseQueryParams(searchParams: URLSearchParams): { performanceTrendQueryParams: PerformanceTrendQueryParams, performanceTrendUIParams: PerformanceTrendUIParams } {
+    let testProjects: string[] = [];
+    let testIdentities: string[] = [];
+    let testTypes: string[] = [];
+    let metricSummaries: string[] = [];
+    let streams: string[] = [];
+    let platforms: string[] = [];
+    let startCommit: number | null = null;
+    let endCommit: number | null = null;
+
+    const getList = (key: string) => {
+        const raw = searchParams.get(key);
+        if (!raw) return [];
+
+        const decompressed = LZString.decompressFromEncodedURIComponent(raw);
+        if (!decompressed) return [];
+
+        return decompressed.split(",").filter(Boolean);
+    };
+
+    const getBoolParam = (key: string): boolean => {
+        return searchParams.get(key)?.toLowerCase() === "true";
+    }
+
+    const querySchemaVersion = Number(searchParams.get(QUERY_SCHEMA_PARAM_NAME)) || 1;
+
+    testProjects = getList(PROJECTS_URL_PARAM_NAME);
+    testIdentities = getList(TEST_IDENTITIES_URL_PARAM_NAME);
+    testTypes = getList(TEST_TYPES_URL_PARAM_NAME);
+    metricSummaries = getList(METRIC_SUMMARIES_URL_PARAM_NAME);
+    streams = getList(STREAMS_URL_PARAM_NAME);
+    platforms = getList(PLATFORM_URL_PARAM_NAME);
+
+    let rawStartCommit = searchParams.get(START_COMMIT_PARAM_NAME);
+    startCommit = rawStartCommit === null || rawStartCommit === "-1" ? null : Number(rawStartCommit);
+
+    let rawEndCommit = searchParams.get(END_COMMIT_PARAM_NAME);
+    endCommit = rawEndCommit === null || rawEndCommit === "-1" ? null : Number(rawEndCommit);
+
+    const curveSmoothing = getBoolParam(UI_CURVE_SMOOTHING_PARAM_NAME);
+    const groupByRaw = searchParams.get(UI_GROUP_BY_PARAM_NAME);
+    const groupBy: GroupByOption = isValidGroupByOption(groupByRaw) ? groupByRaw : 'platform+stream';
+    const viewableProperties = getList(UI_VIEWABLE_PROPERTIES_PARAM_NAME);
+
+    return {
+        performanceTrendQueryParams: { testProjects: testProjects, testIdentities: testIdentities, testTypes: testTypes, metricSummaries: metricSummaries, streams: streams, platforms: platforms, startCommit: startCommit, endCommit: endCommit, querySchemaVersion: querySchemaVersion },
+        performanceTrendUIParams: { applyCurveSmoothing: curveSmoothing, groupBy: groupBy, viewableProperties: viewableProperties }
+    };
+}
+
+const validGroupByOptions = ['platform', 'stream', 'testType', 'platform+stream', 'platform+testType', 'stream+testType', 'platform+stream+testType'] as const;
+
+function isValidGroupByOption(value: string | null): value is GroupByOption {
+    return value !== null && (validGroupByOptions as readonly string[]).includes(value);
+}
+
+/**
+ * Upgrades the query params given the current schema versions.
+ * @param params The params to consider upgrading.
+ * @param currentSchemaVersion The current schema.
+ * @returns The upgraded params.
+ * @todo This is a code-sharing candidate between Performance Trends and Build Health.
+ */
+export function upgradeQueryParams(params: PerformanceTrendQueryParams, currentSchemaVersion: number): PerformanceTrendQueryParams {
+    return params;
+}
+
+/**
+ * Loads options from query params into the options writer.
+ * @param params The params to use to initialize the options system with.
+ * @param uiParams The UI params to use to initialize the options system with.
+ * @param options The options controller to perform the operations.
+ * @param handler The options data handler to use in acquiring subsequent data.
+ * @todo This is a code-sharing candidate between Performance Trends and Build Health. 
+ */
+export async function loadOptionsFromParams(
+    params: PerformanceTrendQueryParams,
+    uiParams: PerformanceTrendUIParams,
+    options: PerformanceTrendOptionsController,
+    handler: PerformanceTrendOptionsDataHandler) {
+
+    let optionsWriter: IPerformanceTrendOptionWriter = options.getTransactionSession();
+
+    // Handle UI first
+    options.setCurveGraphLine(uiParams.applyCurveSmoothing);
+    options.setGroupByOption(uiParams.groupBy);
+
+    // Load viewable properties from URL
+    if (uiParams.viewableProperties.length > 0) {
+        for (const compositeKey of uiParams.viewableProperties) {
+            const decoded = decodeViewablePropertyKey(compositeKey);
+            if (decoded) {
+                options.toggleViewableProperty(compositeKey, decoded.propertyKey, true);
+            }
+        }
+    }
+
+    // We currently support single test project selection
+    if (params.testProjects.length === 1) {
+        runInAction(() => {
+            optionsWriter.toggleSingleProject(params.testProjects[0], params.testProjects[0]);
+        });
+    }
+
+    // We currently support single test identity selection
+    if (params.testIdentities.length === 1) {
+        runInAction(() => {
+            let decodedTestIdentity: DecodedTestIdentityKey | null = decodeTestIdentity(params.testIdentities[0]);
+
+            if (decodedTestIdentity !== null) {
+                optionsWriter.toggleSingleTestIdentity(params.testIdentities[0], decodedTestIdentity.testIdentity);
+            } else {
+                console.warn(`[loadOptionsFromParams] Discarding option param for test identities: ${params.testIdentities[0]}`);
+            }
+        });
+    }
+
+    // We currently support single metric type selection
+    if (params.metricSummaries.length === 1) {
+        runInAction(() => {
+            let decodedMetricType: { testProject: string, testIdentity: string, summaryType: string } | null = decodeMetricTypeKey(params.metricSummaries[0]);
+
+            if (decodedMetricType !== null) {
+                optionsWriter.toggleSingleMetricSummaryType(params.metricSummaries[0], decodedMetricType.summaryType);
+            } else {
+                console.warn(`[loadOptionsFromParams] Discarding option param for metric types: ${params.metricSummaries[0]}`);
+            }
+        });
+    }
+
+    // Support multiple test type selections
+    if (params.testTypes.length > 0) {
+        runInAction(() => {
+            for (let idx: number = 0; idx < params.testTypes.length; ++idx) {
+                let testType: string = params.testTypes[idx];
+                const decodedTestType: DecodedTestTypeKey | null = decodeTestTypeKey(testType);
+
+                if (decodedTestType !== null) {
+                    optionsWriter.toggleTestType(testType, decodedTestType.testType, true);
+                } else {
+                    console.warn(`[loadOptionsFromParams] Discarding option param for test type: ${testType}`);
+                }
+            }
+        });
+    }
+
+    if (params.streams.length >= 1) {
+        runInAction(() => {
+            for (let idx: number = 0; idx < params.streams.length; ++idx) {
+                const decodedStream: DecodedStreamKey | null = decodeStreamKey(params.streams[idx]);
+
+                if (decodedStream !== null) {
+                    optionsWriter.toggleStream(params.streams[idx], decodedStream.stream, true);
+                } else {
+                    console.warn(`[loadOptionsFromParams] Discarding option param for stream: ${params.streams[idx]}`);
+                }
+            }
+        });
+    }
+
+    if (params.platforms.length >= 1) {
+        runInAction(() => {
+            for (let idx: number = 0; idx < params.platforms.length; ++idx) {
+                const decodedPlatform: DecodedPlatformKey | null = decodePlatformKey(params.platforms[idx]);
+
+                if (decodedPlatform !== null) {
+                    optionsWriter.togglePlatform(params.platforms[idx], decodedPlatform.platform, true);
+                } else {
+                    console.warn(`[loadOptionsFromParams] Discarding option param for platform: ${params.platforms[idx]}`);
+                }
+            }
+        });
+    }
+
+    if (params.startCommit !== null) {
+        runInAction(() => {
+            optionsWriter.toggleStartCommitRange(params.startCommit!);
+        });
+    }
+
+    if (params.endCommit !== null) {
+        runInAction(() => {
+            optionsWriter.toggleEndCommitRange(params.endCommit!);
+        });
+    }
+
+    await handler.requestHierarchicalRefresh();
+
+    options.commitTransactionSession();
+}
+
+/**
+ * Serializes the options state into a navigation query.
+ * @param optionsState The options to use for serialization.
+ * @param uiOptionsState The ui options to use for serialization.
+ * @param querySchemaVersion The current schema version.
+ * @param optionsDataHandler The data handler to use in performing serialization.
+ * @returns The serialized URLSearchParams that describes the option set.
+ * @todo This is a code-sharing candidate between Performance Trends and Build Health.
+ */
+export function toNavigationQuery(optionsState: PerformanceTrendOptionsState, uiOptionsState: PerformanceTrendUIOptionsState, querySchemaVersion: number, optionsDataHandler: PerformanceTrendOptionsDataHandler): URLSearchParams {
+    const params = new URLSearchParams(location.search);
+    let projects: string = "";
+    let testIdentities: string = "";
+    let testTypes: string = "";
+    let metricSummaries: string = "";
+    let streams: string = "";
+    let platforms: string = "";
+
+    projects = Object.keys(optionsState.enabledProjects).join(",");
+    testIdentities = Object.keys(optionsState.enabledTestIdentities).join(",");
+    testTypes = Object.keys(optionsState.enabledTestTypes).join(",");
+    metricSummaries = Object.keys(optionsState.enabledMetricSummaryTypes).join(",");
+    streams = Object.keys(optionsState.enabledStreams).join(",");
+    platforms = Object.keys(optionsState.enabledPlatforms).join(",");
+
+    params.set(PROJECTS_URL_PARAM_NAME, LZString.compressToEncodedURIComponent(projects));
+    params.set(TEST_IDENTITIES_URL_PARAM_NAME, LZString.compressToEncodedURIComponent(testIdentities));
+    params.set(TEST_TYPES_URL_PARAM_NAME, LZString.compressToEncodedURIComponent(testTypes));
+    params.set(METRIC_SUMMARIES_URL_PARAM_NAME, LZString.compressToEncodedURIComponent(metricSummaries));
+    params.set(STREAMS_URL_PARAM_NAME, LZString.compressToEncodedURIComponent(streams));
+    params.set(PLATFORM_URL_PARAM_NAME, LZString.compressToEncodedURIComponent(platforms));
+    params.set(QUERY_SCHEMA_PARAM_NAME, querySchemaVersion.toString());
+
+    params.set(START_COMMIT_PARAM_NAME, optionsState.enabledStartCommit !== null ? optionsState.enabledStartCommit.toString() : "-1");
+    params.set(END_COMMIT_PARAM_NAME, optionsState.enabledEndCommit !== null ? optionsState.enabledEndCommit.toString() : "-1");
+
+    // Handle UI Option State
+    params.set(UI_CURVE_SMOOTHING_PARAM_NAME, String(uiOptionsState.curveGraphLine));
+    params.set(UI_GROUP_BY_PARAM_NAME, uiOptionsState.groupByOption);
+
+    // Serialize viewable properties
+    const viewableProperties = Object.keys(uiOptionsState.enabledViewableProperties).join(",");
+    if (viewableProperties) {
+        params.set(UI_VIEWABLE_PROPERTIES_PARAM_NAME, LZString.compressToEncodedURIComponent(viewableProperties));
+    } else {
+        params.delete(UI_VIEWABLE_PROPERTIES_PARAM_NAME);
+    }
+
+    return params;
+}
+
+// #endregion -- Query Serialization & Deserialization --
\ No newline at end of file
diff --git a/Engine/Source/Programs/Horde/HordeDashboard/plugins/performancetrends/components/PlatformDropdownComponent.tsx b/Engine/Source/Programs/Horde/HordeDashboard/plugins/performancetrends/components/PlatformDropdownComponent.tsx
new file mode 100644
--- /dev/null
+++ b/Engine/Source/Programs/Horde/HordeDashboard/plugins/performancetrends/components/PlatformDropdownComponent.tsx
@@ -0,0 +1,163 @@
+// Copyright Epic Games, Inc. All Rights Reserved.
+
+import { Stack, Label, IDropdownOption, SpinnerSize, Spinner } from "@fluentui/react";
+import { observer } from "mobx-react-lite";
+import { PerformanceTrendOptionsController } from "../filters/PerformanceTrendOptionsController";
+import { DataHandlerRefreshRequest, expandRefreshRequest, PerformanceTrendOptionsDataHandler } from "../filters/PerformanceTrendOptionsDataHandler";
+import { PerformanceTrendOptionList, PerformanceTrendOptionData, BasicListParameter } from "./DropdownComponent";
+import { GetTestProjectPlatformResponse } from "../responses/GetTestProjectResponse";
+import { useCallback } from "react";
+import { MultiListParameter } from "./MultilistDropdownComponent";
+import { encodePlatformKey, KEY_SEPARATOR } from "../responses/FilterKeys";
+
+/**
+ * Component used to model all known platform types (for the currently selected test project, test identity, summary table), as a single select dropdown.
+ * @return The react component.
+ */
+export const PlatformDropdownSingle: React.FC<{ handler: PerformanceTrendOptionsDataHandler, performanceTrendOptions: PerformanceTrendOptionsController }> = observer(function ConstructProjectDropdownSingle({ handler, performanceTrendOptions }) {
+    if (handler.isLoadingPlatforms) {
+        return (
+            <Stack>
+                <Label>Platform</Label>
+                <Spinner size={SpinnerSize.small} label="Loading platforms for project..." />
+            </Stack>
+        );
+    }
+
+    const platformDropdownParams: PerformanceTrendOptionList = {
+        label: "",
+        items: [
+        ],
+        tooltip: "Select an option from the dropdown."
+    };
+
+    handler.platforms.forEach((platform: GetTestProjectPlatformResponse) => {
+        platform.platforms.forEach((platformStr: string) => {
+            let listParam: PerformanceTrendOptionData = {
+                id: encodePlatformKey(platform.owningTestProject, platformStr),
+                text: platformStr,
+            };
+            platformDropdownParams.items.push(listParam);
+        })
+    });
+
+    platformDropdownParams.items.sort((a, b) => a.text.localeCompare(b.text));
+
+    let selectedPlatformKey = Object.keys(performanceTrendOptions.state.enabledPlatforms)[0] ?? null;
+    let request = expandRefreshRequest({ commits: true });
+
+    return platformDropdownParams.items.length > 0 && (
+        <Stack>
+            <Label>Platform</Label>
+            <BasicListParameter placeholder={"Select a platform"} handler={handler} performanceTrendOptions={performanceTrendOptions} params={platformDropdownParams} selectedKey={selectedPlatformKey} request={request} onChange={(option: IDropdownOption<PerformanceTrendOptionData>) => {
+                const key = option.key as string;
+                performanceTrendOptions.getTransactionSession().toggleSinglePlatform(key, option.text);
+            }} />
+        </Stack>
+    );
+});
+
+/**
+ * Component used to model all known platform types (for the currently selected test project, test identity, summary table), as a multi select dropdown.
+ * @return The react component.
+ */
+export const PlatformDropdownMulti: React.FC<{ handler: PerformanceTrendOptionsDataHandler, performanceTrendOptions: PerformanceTrendOptionsController }> = observer(function ConstructProjectDropdown({ handler, performanceTrendOptions }) {
+    const handleSelectAll = useCallback(
+        (
+            option: IDropdownOption<PerformanceTrendOptionData>,
+            optionSelectAllContext?: IDropdownOption<PerformanceTrendOptionData>[]
+        ) => {
+            const key = option.key as string;
+            const [, projectId] = key.split(KEY_SEPARATOR);
+            const isSelected = !!option.selected;
+
+            // If there is no provided context, simply obtain all the known platforms of the (project).
+            const platformIds: { key: string, platform: string }[] =
+                optionSelectAllContext === undefined
+                    ? handler.platforms.flatMap(x => x.platforms.map(platformStr => ({ key: encodePlatformKey(x.owningTestProject, platformStr), platform: platformStr })))
+                    : optionSelectAllContext
+                        .filter((opt) => {
+                            if (!opt.data) return false;
+
+                            return (
+                                opt.data.group === projectId
+                            );
+                        })
+                        .map((x) => ({ key: x.data!.id, platform: x.data!.text! }));
+
+           // Do not close the transaction session; the user still has the modal open
+            performanceTrendOptions.getTransactionSession().toggleAllPlatforms(platformIds, isSelected);
+        },
+        [performanceTrendOptions, handler]
+    );
+
+    if (handler.isLoadingPlatforms) {
+        return (
+            <Stack>
+                <Label>Platform</Label>
+                <Spinner size={SpinnerSize.small} label="Loading platforms for project..." />
+            </Stack>
+        );
+    }
+
+    const platformDropdownParams: PerformanceTrendOptionList = {
+        label: "Select Platform(s)",
+        items: [
+        ],
+        tooltip: "Select an option from the dropdown."
+    };
+
+    // Build a map of platform -> all streams for that platform
+    const platformStreamMap = new Map<string, Set<string>>();
+    for (const p of handler.platforms) {
+        for (const platformStr of p.platforms) {
+            if (!platformStreamMap.has(platformStr)) {
+                platformStreamMap.set(platformStr, new Set<string>());
+            }
+            platformStreamMap.get(platformStr)!.add(p.owningTestProject.computedStream);
+        }
+    }
+
+    // Build unique platform items by flattening all platform arrays
+    const seenPlatforms = new Set<string>();
+    const platformItems: PerformanceTrendOptionData[] = [];
+
+    for (const p of handler.platforms) {
+        for (const platformStr of p.platforms) {
+            if (!seenPlatforms.has(platformStr)) {
+                seenPlatforms.add(platformStr);
+                const streams = platformStreamMap.get(platformStr);
+                const streamList = streams ? Array.from(streams).sort().join('\n  - ') : '';
+
+                platformItems.push({
+                    id: encodePlatformKey(p.owningTestProject, platformStr),
+                    text: platformStr,
+                    group: p.owningTestProject.testName,
+                    tooltip: streamList ? `Streams:\n  - ${streamList}` : undefined,
+                });
+            }
+        }
+    }
+
+    platformDropdownParams.items = platformItems.sort((a, b) => a.text.localeCompare(b.text));
+
+    let selectedPlatformKeys: string[] = Object.keys(performanceTrendOptions.state.enabledPlatforms);
+    let request: DataHandlerRefreshRequest = expandRefreshRequest({ commits: true });
+
+    return platformDropdownParams.items.length > 0 && (
+        <Stack>
+            <Label>Platform(s)</Label>
+            <Stack horizontal tokens={{ childrenGap: 8 }} verticalAlign="end">
+                <MultiListParameter handler={handler} performanceTrendOptions={performanceTrendOptions} params={platformDropdownParams} selectedKeys={selectedPlatformKeys} request={request}
+                    enabledSelectAll={true}
+                    onSelectAll={handleSelectAll}
+                    onChange={(option: IDropdownOption<PerformanceTrendOptionData>) => {
+                        const key = option.key as string;
+                        const isSelected = !!option.selected;
+
+                        performanceTrendOptions.getTransactionSession().togglePlatform(key, option.text, isSelected);
+                    }} />
+            </Stack>
+        </Stack>
+    );
+});
\ No newline at end of file
diff --git a/Engine/Source/Programs/Horde/HordeDashboard/plugins/performancetrends/components/StreamDropdownComponent.tsx b/Engine/Source/Programs/Horde/HordeDashboard/plugins/performancetrends/components/StreamDropdownComponent.tsx
new file mode 100644
--- /dev/null
+++ b/Engine/Source/Programs/Horde/HordeDashboard/plugins/performancetrends/components/StreamDropdownComponent.tsx
@@ -0,0 +1,108 @@
+import { IDropdownOption, Stack, Label } from "@fluentui/react";
+import { observer } from "mobx-react-lite";
+import { useCallback } from "react";
+import { PerformanceTrendOptionsController } from "../filters/PerformanceTrendOptionsController";
+import { PerformanceTrendOptionsDataHandler, DataHandlerRefreshRequest, expandRefreshRequest } from "../filters/PerformanceTrendOptionsDataHandler";
+import { PerformanceTrendOptionData, PerformanceTrendOptionList } from "./DropdownComponent";
+import { MultiListParameter } from "./MultilistDropdownComponent";
+import { KEY_SEPARATOR, encodeStreamKey } from "../responses/FilterKeys";
+
+/**
+ * Component used to model all known streams(for the currently selected test project, test identity, summary table), as a single select dropdown.
+ * @return The react component.
+ */
+export const StreamDropdownMulti: React.FC<{ handler: PerformanceTrendOptionsDataHandler, performanceTrendOptions: PerformanceTrendOptionsController }> = observer(function ConstructProjectDropdown({ handler, performanceTrendOptions }) {
+    const handleSelectAll = useCallback(
+        (
+            option: IDropdownOption<PerformanceTrendOptionData>,
+            optionSelectAllContext?: IDropdownOption<PerformanceTrendOptionData>[]
+        ) => {
+            const key = option.key as string;
+            const [, projectId] = key.split(KEY_SEPARATOR);
+            const isSelected = !!option.selected;
+
+            // If there is no provided context, simply obtain all the known unique streams of the (project).
+            const seenStreamsInSelectAll = new Set<string>();
+            const uniqueStreams = handler.streams.filter(x => {
+                if (seenStreamsInSelectAll.has(x.computedStream)) {
+                    return false;
+                }
+                seenStreamsInSelectAll.add(x.computedStream);
+                return true;
+            });
+            const streamIds: { key: string, stream: string }[] =
+                optionSelectAllContext === undefined
+                    ? uniqueStreams.map(x => ({ key: encodeStreamKey(x), stream: x.computedStream }))
+                    : optionSelectAllContext
+                        .filter((opt) => {
+                            if (!opt.data) return false;
+
+                            return (
+                                opt.data.group === projectId
+                            );
+                        })
+                        .map((x) => ({ key: x.data!.id, stream: x.data!.text! }));
+
+            // Do not close the transaction session; the user still has the modal open
+            performanceTrendOptions.getTransactionSession().toggleAllStreams(streamIds, isSelected);
+        },
+        [performanceTrendOptions, handler]
+    );
+
+    const streamDropdownParams: PerformanceTrendOptionList = {
+        label: "Select Streams(s)",
+        items: [
+        ],
+        tooltip: "Select an option from the dropdown."
+    };
+
+    // Build a map of stream -> all test types for that stream
+    const streamTestTypeMap = new Map<string, Set<string>>();
+    for (const stream of handler.streams) {
+        if (!streamTestTypeMap.has(stream.computedStream)) {
+            streamTestTypeMap.set(stream.computedStream, new Set<string>());
+        }
+        streamTestTypeMap.get(stream.computedStream)!.add(stream.testType);
+    }
+
+    // Build unique stream items
+    const seenStreams = new Set<string>();
+    const streamItems: PerformanceTrendOptionData[] = [];
+
+    for (const stream of handler.streams) {
+        if (!seenStreams.has(stream.computedStream)) {
+            seenStreams.add(stream.computedStream);
+            const testTypes = streamTestTypeMap.get(stream.computedStream);
+            const testTypeList = testTypes ? Array.from(testTypes).sort().join('\n  - ') : '';
+
+            streamItems.push({
+                id: encodeStreamKey(stream),
+                text: stream.computedStream,
+                group: stream.testName,
+                tooltip: testTypeList ? `Test Types:\n  - ${testTypeList}` : undefined,
+            });
+        }
+    }
+
+    streamDropdownParams.items = streamItems.sort((a, b) => a.text.localeCompare(b.text));
+
+    let selectedStreamKeys: string[] = Object.keys(performanceTrendOptions.state.enabledStreams);
+    let request: DataHandlerRefreshRequest = expandRefreshRequest({ platforms: true });
+
+    return streamDropdownParams.items.length > 0 && (
+        <Stack>
+            <Label>Stream(s)</Label>
+            <Stack horizontal tokens={{ childrenGap: 8 }} verticalAlign="end">
+                <MultiListParameter handler={handler} performanceTrendOptions={performanceTrendOptions} params={streamDropdownParams} selectedKeys={selectedStreamKeys} request={request}
+                    enabledSelectAll={true}
+                    onSelectAll={handleSelectAll}
+                    onChange={(option: IDropdownOption<PerformanceTrendOptionData>) => {
+                        const key = option.key as string;
+                        const isSelected = !!option.selected;
+
+                        performanceTrendOptions.getTransactionSession().toggleStream(key, option.text, isSelected);
+                    }} />
+            </Stack>
+        </Stack>
+    );
+});
\ No newline at end of file
diff --git a/Engine/Source/Programs/Horde/HordeDashboard/plugins/performancetrends/components/TestIdentitiesTypeDropdownSingle.tsx b/Engine/Source/Programs/Horde/HordeDashboard/plugins/performancetrends/components/TestIdentitiesTypeDropdownSingle.tsx
new file mode 100644
--- /dev/null
+++ b/Engine/Source/Programs/Horde/HordeDashboard/plugins/performancetrends/components/TestIdentitiesTypeDropdownSingle.tsx
@@ -0,0 +1,45 @@
+// Copyright Epic Games, Inc. All Rights Reserved.
+
+import { Stack, Label, IDropdownOption } from "@fluentui/react";
+import { observer } from "mobx-react-lite";
+import { PerformanceTrendOptionsController } from "../filters/PerformanceTrendOptionsController";
+import { DataHandlerRefreshRequest, expandRefreshRequest, PerformanceTrendOptionsDataHandler } from "../filters/PerformanceTrendOptionsDataHandler";
+import { PerformanceTrendOptionList, PerformanceTrendOptionData, BasicListParameter } from "./DropdownComponent";
+import { GetTestProjectResponse } from "../responses/GetTestProjectResponse";
+import { encodeTestIdentity } from "../responses/FilterKeys";
+
+/**
+ * Component used to model all known test identities (for the currently selected test project), as a single select dropdown.
+ * @return The react component.
+ */
+export const TestIdentitiesTypeDropdownSingle: React.FC<{ handler: PerformanceTrendOptionsDataHandler, performanceTrendOptions: PerformanceTrendOptionsController }> = observer(function ConstructProjectDropdownSingle({ handler, performanceTrendOptions }) {
+    const testIdentityDropdownParams: PerformanceTrendOptionList = {
+        label: "",
+        items: [
+        ],
+        tooltip: "Select an option from the dropdown."
+    };
+
+    testIdentityDropdownParams.items = handler.testIdentities.map((testIdentityParentProject: GetTestProjectResponse) => {
+        let listParam: PerformanceTrendOptionData = {
+            id: encodeTestIdentity(testIdentityParentProject),
+            text: testIdentityParentProject.testIdentity,
+        };
+
+        return listParam;
+    }).sort((a, b) => a.text.localeCompare(b.text));
+
+    let selectedTestIdentityKey = Object.keys(performanceTrendOptions.state.enabledTestIdentities)[0] ?? null;
+    let request: DataHandlerRefreshRequest = expandRefreshRequest({ metricTypes: true });
+
+    return testIdentityDropdownParams.items.length > 0 && (
+        <Stack>
+            <Label>Test Identity</Label>
+            <BasicListParameter placeholder={"Select a test identity"} handler={handler} performanceTrendOptions={performanceTrendOptions} params={testIdentityDropdownParams} selectedKey={selectedTestIdentityKey} request={request} onChange={(option: IDropdownOption<PerformanceTrendOptionData>) => {
+                const key = option.key as string;
+                performanceTrendOptions.getTransactionSession().toggleSingleTestIdentity(key, option.text);
+            }}
+                selectorWidth={400} />
+        </Stack>
+    );
+});
\ No newline at end of file
diff --git a/Engine/Source/Programs/Horde/HordeDashboard/plugins/performancetrends/components/TestProjectDropdownComponent.tsx b/Engine/Source/Programs/Horde/HordeDashboard/plugins/performancetrends/components/TestProjectDropdownComponent.tsx
new file mode 100644
--- /dev/null
+++ b/Engine/Source/Programs/Horde/HordeDashboard/plugins/performancetrends/components/TestProjectDropdownComponent.tsx
@@ -0,0 +1,54 @@
+// Copyright Epic Games, Inc. All Rights Reserved.
+
+import { IDropdownOption, Label, Spinner, SpinnerSize, Stack } from "@fluentui/react";
+import { BasicListParameter, PerformanceTrendOptionData, PerformanceTrendOptionList } from "./DropdownComponent";
+import { PerformanceTrendOptionsController } from "../filters/PerformanceTrendOptionsController";
+import { observer } from "mobx-react-lite";
+import { DataHandlerRefreshRequest, expandRefreshRequest, PerformanceTrendOptionsDataHandler } from "../filters/PerformanceTrendOptionsDataHandler";
+
+/**
+ * Component used to model all known test projects, as a single select dropdown.
+ * @return The react component.
+ */
+export const TestProjectDropdownSingle: React.FC<{ handler: PerformanceTrendOptionsDataHandler, performanceTrendOptions: PerformanceTrendOptionsController }> = observer(function ConstructProjectDropdownSingle({ handler, performanceTrendOptions }) {
+    if (handler.isLoadingProjects) {
+        return (
+            <Stack>
+                <Label>Test Project</Label>
+                <Spinner size={SpinnerSize.small} label="Loading projects..." />
+            </Stack>
+        );
+    }
+
+    const testProjectDropdownParams: PerformanceTrendOptionList = {
+        label: "",
+        items: [
+        ],
+        tooltip: "Select an option from the dropdown."
+    };
+
+    const uniqueTestNames = Array.from(
+        new Set(handler.projectsData.map(p => p.testName))
+    );
+
+    testProjectDropdownParams.items = uniqueTestNames
+        .map(testName => ({
+            id: testName,
+            text: testName,
+        }))
+        .sort((a, b) => a.text.localeCompare(b.text));
+
+    let selectedTestProject = Object.keys(performanceTrendOptions.state.enabledProjects)[0] ?? null;
+
+    let request: DataHandlerRefreshRequest = expandRefreshRequest({ testIdentities: true });
+
+    return (
+        <Stack>
+            <Label>Test Project</Label>
+            <BasicListParameter placeholder={"Select a test project"} handler={handler} performanceTrendOptions={performanceTrendOptions} params={testProjectDropdownParams} selectedKey={selectedTestProject} request={request} onChange={(option: IDropdownOption<PerformanceTrendOptionData>) => {
+                const key = option.key as string;
+                performanceTrendOptions.getTransactionSession().toggleSingleProject(key, option.text);
+            }} />
+        </Stack>
+    );
+});
\ No newline at end of file
diff --git a/Engine/Source/Programs/Horde/HordeDashboard/plugins/performancetrends/components/TestTypeDropdownComponent.tsx b/Engine/Source/Programs/Horde/HordeDashboard/plugins/performancetrends/components/TestTypeDropdownComponent.tsx
new file mode 100644
--- /dev/null
+++ b/Engine/Source/Programs/Horde/HordeDashboard/plugins/performancetrends/components/TestTypeDropdownComponent.tsx
@@ -0,0 +1,118 @@
+// Copyright Epic Games, Inc. All Rights Reserved.
+
+import { Stack, Label, IDropdownOption, Spinner, SpinnerSize } from "@fluentui/react";
+import { observer } from "mobx-react-lite";
+import { PerformanceTrendOptionsController } from "../filters/PerformanceTrendOptionsController";
+import { DataHandlerRefreshRequest, expandRefreshRequest, PerformanceTrendOptionsDataHandler } from "../filters/PerformanceTrendOptionsDataHandler";
+import { PerformanceTrendOptionList, PerformanceTrendOptionData, BasicListParameter } from "./DropdownComponent";
+import { GetTestProjectResponse } from "../responses/GetTestProjectResponse";
+import { encodeTestTypeKey, KEY_SEPARATOR } from "../responses/FilterKeys";
+import { useCallback } from "react";
+import { MultiListParameter } from "./MultilistDropdownComponent";
+
+/**
+ * Component used to model all known test types (for the currently selected test project, and test identity), as a single select dropdown.
+ * @return The react component.
+ */
+export const TestTypeDropdownSingle: React.FC<{ handler: PerformanceTrendOptionsDataHandler, performanceTrendOptions: PerformanceTrendOptionsController }> = observer(function ConstructProjectDropdownSingle({ handler, performanceTrendOptions }) {
+    const testTypeDropdownParams: PerformanceTrendOptionList = {
+        label: "",
+        items: [
+        ],
+        tooltip: "Select an option from the dropdown."
+    };
+
+    testTypeDropdownParams.items = handler.testTypes.map((testType: GetTestProjectResponse) => {
+        let listParam: PerformanceTrendOptionData = {
+            id: encodeTestTypeKey(testType),
+            text: testType.testType,
+        };
+
+        return listParam;
+    }).sort((a, b) => a.text.localeCompare(b.text));
+
+    let selectedTestTypeKey = Object.keys(performanceTrendOptions.state.enabledTestTypes)[0] ?? null;
+    let request: DataHandlerRefreshRequest = expandRefreshRequest({ streams: true });
+
+    return testTypeDropdownParams.items.length > 0 && (
+        <Stack>
+            <Label>Test Type</Label>
+            <BasicListParameter placeholder={"Select a test type"} handler={handler} performanceTrendOptions={performanceTrendOptions} params={testTypeDropdownParams} selectedKey={selectedTestTypeKey} request={request} onChange={(option: IDropdownOption<PerformanceTrendOptionData>) => {
+                const key = option.key as string;
+                performanceTrendOptions.getTransactionSession().toggleSingleTestType(key, option.text);
+            }} />
+        </Stack>
+    );
+});
+
+/**
+ * Component used to model all known test types (for the currently selected test project, test identity, summary table), as a multi select dropdown.
+ * @return The react component.
+ */
+export const TestTypeDropdownMulti: React.FC<{ handler: PerformanceTrendOptionsDataHandler, performanceTrendOptions: PerformanceTrendOptionsController }> = observer(function ConstructProjectDropdown({ handler, performanceTrendOptions }) {
+    const handleSelectAll = useCallback(
+        (
+            option: IDropdownOption<PerformanceTrendOptionData>,
+            optionSelectAllContext?: IDropdownOption<PerformanceTrendOptionData>[]
+        ) => {
+            const key = option.key as string;
+            const [, projectId] = key.split(KEY_SEPARATOR);
+            const isSelected = !!option.selected;
+
+            // If there is no provided context, simply obtain all the known test types of the (project).
+            const testTypeIds: { key: string, testType: string }[] =
+                optionSelectAllContext === undefined
+                    ? handler.testTypes.map(x => ({ key: encodeTestTypeKey(x), testType: x.testType }))
+                    : optionSelectAllContext
+                        .filter((opt) => {
+                            if (!opt.data) return false;
+
+                            return (
+                                opt.data.group === projectId
+                            );
+                        })
+                        .map((x) => ({ key: x.data!.id, testType: x.data!.text! }));
+
+            // Do not close the transaction session; the user still has the modal open
+            performanceTrendOptions.getTransactionSession().toggleAllTestTypes(testTypeIds, isSelected);
+        },
+        [performanceTrendOptions, handler]
+    );
+
+    const testTypeDropdownParams: PerformanceTrendOptionList = {
+        label: "Select Test Type(s)",
+        items: [
+        ],
+        tooltip: "Select an option from the dropdown."
+    };
+
+    testTypeDropdownParams.items = handler.testTypes.map((testType: GetTestProjectResponse) => {
+        let listParam: PerformanceTrendOptionData = {
+            id: encodeTestTypeKey(testType),
+            group: testType.testName,
+            text: testType.testType,
+        };
+
+        return listParam;
+    }).sort((a, b) => a.text.localeCompare(b.text));
+
+    let selectedTestTypeKeys: string[] = Object.keys(performanceTrendOptions.state.enabledTestTypes);
+    let request: DataHandlerRefreshRequest = expandRefreshRequest({ streams: true });
+
+    return testTypeDropdownParams.items.length > 0 && (
+        <Stack>
+            <Label>Test Type(s)</Label>
+            <Stack horizontal tokens={{ childrenGap: 8 }} verticalAlign="end">
+                <MultiListParameter handler={handler} performanceTrendOptions={performanceTrendOptions} params={testTypeDropdownParams} selectedKeys={selectedTestTypeKeys} request={request}
+                    enabledSelectAll={true}
+                    onSelectAll={handleSelectAll}
+                    onChange={(option: IDropdownOption<PerformanceTrendOptionData>) => {
+                        const key = option.key as string;
+                        const isSelected = !!option.selected;
+
+                        performanceTrendOptions.getTransactionSession().toggleTestType(key, option.text, isSelected);
+                    }} />
+            </Stack>
+        </Stack>
+    );
+});
\ No newline at end of file
diff --git a/Engine/Source/Programs/Horde/HordeDashboard/plugins/performancetrends/components/ViewablePropertiesDropdownComponent.tsx b/Engine/Source/Programs/Horde/HordeDashboard/plugins/performancetrends/components/ViewablePropertiesDropdownComponent.tsx
new file mode 100644
--- /dev/null
+++ b/Engine/Source/Programs/Horde/HordeDashboard/plugins/performancetrends/components/ViewablePropertiesDropdownComponent.tsx
@@ -0,0 +1,129 @@
+// Copyright Epic Games, Inc. All Rights Reserved.
+
+import { Stack, Label, IDropdownOption } from "@fluentui/react";
+import { observer } from "mobx-react-lite";
+import { PerformanceTrendOptionsController } from "../filters/PerformanceTrendOptionsController";
+import { IMetricViewGenerator } from "../viewgenerators/PerformanceTrendRenderTypes";
+import { MetricConstraint } from "../metrictypes/PerformanceTrendsTypes";
+import { MultiListParameter } from "./MultilistDropdownComponent";
+import { PerformanceTrendOptionList, PerformanceTrendOptionData } from "./DropdownComponent";
+import { PerformanceTrendOptionsDataHandler } from "../filters/PerformanceTrendOptionsDataHandler";
+import { useCallback, useEffect, useMemo } from "react";
+import { decodeTestTypeKey, decodeViewablePropertyKey, encodeViewablePropertyKey, KEY_SEPARATOR } from "../responses/FilterKeys";
+
+/**
+ * Dropdown for selecting which metric properties to display in the view.
+ * Groups properties by test type when multiple test types are selected.
+ * Only visible when at least one test type is selected.
+ */
+export const ViewablePropertiesDropdown: React.FC<{ viewGenerator: IMetricViewGenerator<MetricConstraint> | undefined, handler: PerformanceTrendOptionsDataHandler, performanceTrendOptions: PerformanceTrendOptionsController }> = observer(function ViewablePropertiesDropdown({ viewGenerator, handler, performanceTrendOptions }) {
+
+    // Get selected test types from published state (synchronized after commit)
+    // Use optionsChangeVersion as dependency to properly react to MobX state changes
+    const selectedTestTypes = useMemo(() => {
+        const testTypeKeys = performanceTrendOptions.state.publishedEnabledTestTypeKeys;
+        return testTypeKeys.map(key => {
+            const decoded = decodeTestTypeKey(key);
+            return decoded ? { key, testType: decoded.testType } : null;
+        }).filter((x): x is { key: string, testType: string } => x !== null);
+    }, [performanceTrendOptions.optionsChangeVersion]);
+
+    // Build dropdown items from viewGenerator's viewable properties, grouped by test type
+    const dropdownParams = useMemo((): PerformanceTrendOptionList => {
+        const items: PerformanceTrendOptionData[] = [];
+
+        if (viewGenerator && selectedTestTypes.length > 0) {
+            const properties = viewGenerator.getViewableProperties();
+
+            for (const { testType } of selectedTestTypes) {
+                for (const [fieldName, meta] of properties) {
+                    const compositeKey = encodeViewablePropertyKey(testType, fieldName);
+                    items.push({
+                        id: compositeKey,
+                        text: meta.unit ? `${meta.label} (${meta.unit})` : meta.label,
+                        group: testType,
+                        tooltip: meta.category,
+                    });
+                }
+            }
+        }
+
+        // Sort by group (test type) then by label
+        items.sort((a, b) => {
+            const groupCompare = (a.group ?? '').localeCompare(b.group ?? '');
+            if (groupCompare !== 0) return groupCompare;
+            return a.text.localeCompare(b.text);
+        });
+
+        return {
+            label: "Select Properties",
+            items,
+            tooltip: "Select which properties to display in the graphs for each test type."
+        };
+    }, [viewGenerator, selectedTestTypes]);
+
+    // Initialize default properties when test types change
+    useEffect(() => {
+        if (viewGenerator && selectedTestTypes.length > 0) {
+            const defaultProps = viewGenerator.getDefaultVisibleProperties();
+            for (const { testType } of selectedTestTypes) {
+                performanceTrendOptions.initializeDefaultViewableProperties(testType, defaultProps);
+            }
+        }
+    }, [viewGenerator, selectedTestTypes, performanceTrendOptions]);
+
+    const handleSelectAll = useCallback(
+        (option: IDropdownOption<PerformanceTrendOptionData>, optionSelectAllContext?: IDropdownOption<PerformanceTrendOptionData>[]) => {
+            const isSelected = !!option.selected;
+            const key = option.key as string;
+            const group = key.split(KEY_SEPARATOR)[1];
+
+            const properties: { key: string, propertyKey: string }[] = optionSelectAllContext === undefined
+                ? dropdownParams.items.map(x => {
+                    const decoded = decodeViewablePropertyKey(x.id);
+                    return { key: x.id, testType: decoded?.testType, propertyKey: decoded?.propertyKey ?? '' };
+                }).filter(x => x.testType === group)
+                : optionSelectAllContext
+                    .filter(opt => opt.data?.group === group)
+                    .map(x => {
+                        const decoded = decodeViewablePropertyKey(x.data!.id);
+                        return { key: x.data!.id, propertyKey: decoded?.propertyKey ?? '' };
+                    });
+
+            performanceTrendOptions.toggleAllViewableProperties(properties, isSelected);
+        },
+        [performanceTrendOptions, dropdownParams]
+    );
+
+    // Don't show if no test types selected or no view generator
+    if (!viewGenerator || selectedTestTypes.length === 0 || dropdownParams.items.length === 0) {
+        return null;
+    }
+
+    const selectedKeys: string[] = Object.keys(performanceTrendOptions.uiState.enabledViewableProperties);
+
+    return (
+        <Stack>
+            <Label>Visible Properties</Label>
+            <Stack horizontal tokens={{ childrenGap: 8 }} verticalAlign="end">
+                <MultiListParameter
+                    handler={handler}
+                    performanceTrendOptions={performanceTrendOptions}
+                    params={dropdownParams}
+                    selectedKeys={selectedKeys}
+                    request={{}} // No data refresh needed - view-only filter
+                    enabledSelectAll={true}
+                    onSelectAll={handleSelectAll}
+                    onChange={(option: IDropdownOption<PerformanceTrendOptionData>) => {
+                        const key = option.key as string;
+                        const decoded = decodeViewablePropertyKey(key);
+                        if (decoded) {
+                            const isSelected = !!option.selected;
+                            performanceTrendOptions.toggleViewableProperty(key, decoded.propertyKey, isSelected);
+                        }
+                    }}
+                />
+            </Stack>
+        </Stack>
+    );
+});
diff --git a/Engine/Source/Programs/Horde/HordeDashboard/plugins/performancetrends/docs/PerformanceTrends.md b/Engine/Source/Programs/Horde/HordeDashboard/plugins/performancetrends/docs/PerformanceTrends.md
new file mode 100644
--- /dev/null
+++ b/Engine/Source/Programs/Horde/HordeDashboard/plugins/performancetrends/docs/PerformanceTrends.md
@@ -0,0 +1,397 @@
+# Adding New Metric Types to Performance Trends
+
+This guide walks through adding a new metric type to the Performance Trends dashboard, covering both server-side (C#) and client-side (TypeScript/React) implementation.
+
+## Overview
+
+The Performance Trends system uses a registry-based architecture where metric types are:
+1. **Server-side**: Queried via `AbstractPerformanceSummaryTrendHandler` implementations
+2. **Client-side**: Visualized via `IMetricViewGenerator<T>` implementations
+
+Both sides are linked by a shared `metricType` string identifier.
+
+---
+
+## Server-Side Implementation
+
+### Step 1: Define the Telemetry Record
+
+Create a new record class in `EpicGames.Analytics.PerformanceTrends/`:
+
+```csharp
+// MemoryStatsTelemetryRecord.cs
+
+using System.ComponentModel.DataAnnotations.Schema;
+using EpicGames.Analytics.PerformanceTrends;
+
+namespace EpicGames.Analytics
+{
+    /// <summary>
+    /// Telemetry record for memory performance metrics.
+    /// </summary>
+    [AnalyticsTableGen]
+    [Table("schema.memory_stats_summary")]
+    public record MemoryStatsTelemetryRecord : PerformanceTrendTelemetry
+    {
+        [Column("heap_used_mb")]
+        public float HeapUsedMb { get; init; }
+
+        [Column("texture_memory_mb")]
+        public float TextureMemoryMb { get; init; }
+
+        [Column("mesh_memory_mb")]
+        public float MeshMemoryMb { get; init; }
+
+        [Column("audio_memory_mb")]
+        public float AudioMemoryMb { get; init; }
+
+        [Column("total_allocated_mb")]
+        public float TotalAllocatedMb { get; init; }
+
+        public MemoryStatsTelemetryRecord() : base() { }
+    }
+}
+```
+
+> **Note**: The `[AnalyticsTableGen]` attribute generates `EpicGames_Analytics_MemoryStatsTelemetryRecordGen` with type-safe column constants and Dapper mappings.
+
+### Step 2: Create the Handler
+
+Create a handler in `HordeServer.Epic.EpicSandbox/PerformanceTrends/`:
+
+```csharp
+// MemoryStatsPerformanceSummaryHandler.cs
+
+using System.Data;
+using System.Data.Common;
+using System.Data.Odbc;
+using EpicGames.Analytics;
+using EpicGames.Analytics.Generated;
+using EpicGames.Analytics.PerformanceTrends;
+using HordeServer.Analytics;
+using Microsoft.Extensions.Logging;
+
+namespace HordeServer.LicenseePromotionCandidate.PerformanceTrends
+{
+    /// <summary>
+    /// Summary handler for memory statistics metrics.
+    /// </summary>
+    public class MemoryStatsPerformanceSummaryHandler : AbstractPerformanceSummaryTrendHandler
+    {
+        // Register Dapper column mappings at startup
+        static MemoryStatsPerformanceSummaryHandler()
+        {
+            ColumnHandlers.RegisterDapperColumnMapping<MemoryStatsTelemetryRecord>();
+        }
+
+        /// <inheritdoc/>
+        public override string PerformanceSummaryType => "MemoryStats";
+
+        /// <inheritdoc/>
+        public override async Task<IEnumerable<PerformanceTrendTelemetry>> ProcessGeneralMetricRequest(
+            OdbcConnection dbConnection,
+            PerformanceTrendFilter filter,
+            ILogger logger,
+            CancellationToken cancellationToken)
+        {
+            List<MemoryStatsTelemetryRecord> results = [];
+
+            using (OdbcCommand cmd = dbConnection.CreateCommand())
+            {
+                string whereSql = BuildWhereClause(filter, cmd);
+                string sql = $@"
+                    SELECT
+                        {EpicGames_Analytics_MemoryStatsTelemetryRecordGen.HeapUsedMb},
+                        {EpicGames_Analytics_MemoryStatsTelemetryRecordGen.TextureMemoryMb},
+                        {EpicGames_Analytics_MemoryStatsTelemetryRecordGen.MeshMemoryMb},
+                        {EpicGames_Analytics_MemoryStatsTelemetryRecordGen.AudioMemoryMb},
+                        {EpicGames_Analytics_MemoryStatsTelemetryRecordGen.TotalAllocatedMb},
+                        {EpicGames_Analytics_PerformanceTrends_PerformanceTrendTelemetryGen.SummaryName},
+                        {EpicGames_Analytics_PerformanceTrends_PerformanceTrendTelemetryGen.TestName},
+                        {EpicGames_Analytics_PerformanceTrends_PerformanceTrendTelemetryGen.CommitId}
+                    FROM {EpicGames_Analytics_MemoryStatsTelemetryRecordGen.TableName}
+                    {whereSql}
+                    ORDER BY {EpicGames_Analytics_PerformanceTrends_PerformanceTrendTelemetryGen.CommitId} DESC
+                    LIMIT {filter.RecordCount}";
+
+#pragma warning disable CA2100
+                cmd.CommandText = sql;
+#pragma warning restore CA2100
+                cmd.CommandType = CommandType.Text;
+
+                using (DbDataReader reader = await cmd.ExecuteReaderAsync(cancellationToken))
+                {
+                    results = reader.Parse<MemoryStatsTelemetryRecord>().AsList();
+                    await reader.CloseAsync();
+                }
+            }
+
+            return results;
+        }
+
+        /// <inheritdoc/>
+        protected override string GetTableName()
+        {
+            return EpicGames_Analytics_MemoryStatsTelemetryRecordGen.TableName;
+        }
+
+        /// <inheritdoc/>
+        protected override IReadOnlyList<PerformanceTrendTelemetry> ParseTypeFromReader(DbDataReader reader)
+        {
+            return reader.Parse<MemoryStatsTelemetryRecord>().ToList();
+        }
+
+        private static string BuildWhereClause(PerformanceTrendFilter filter, OdbcCommand cmd)
+        {
+            List<string> where = new();
+
+            if (!string.IsNullOrEmpty(filter.TestProject))
+            {
+                where.Add($"{EpicGames_Analytics_MemoryStatsTelemetryRecordGen.TestName} = ?");
+                var param = cmd.CreateParameter();
+                param.OdbcType = OdbcType.VarChar;
+                param.Value = filter.TestProject;
+                cmd.Parameters.Add(param);
+            }
+
+            if (!string.IsNullOrEmpty(filter.TestIdentity))
+            {
+                where.Add($"{EpicGames_Analytics_MemoryStatsTelemetryRecordGen.TestIdentity} = ?");
+                var param = cmd.CreateParameter();
+                param.OdbcType = OdbcType.VarChar;
+                param.Value = filter.TestIdentity;
+                cmd.Parameters.Add(param);
+            }
+
+            return where.Count > 0 ? "WHERE " + string.Join(" AND ", where) : string.Empty;
+        }
+    }
+}
+```
+
+### Step 3: Register in Dependency Injection
+
+In your service configuration:
+
+```csharp
+services.AddSingleton<AbstractPerformanceSummaryTrendHandler, MemoryStatsPerformanceSummaryHandler>();
+```
+
+---
+
+## Client-Side Implementation
+
+### Step 4: Define the Metric Data Class
+
+Create in `PerformanceTrends/metrictypes/`:
+
+```typescript
+// MemoryStatsData.ts
+
+import { MemoryStatsViewGenerator } from "../viewgenerators/MemoryStatsViewGenerator";
+import { ISummaryMetric, MetricTypeRegistry, PerformanceTrendContext } from "./PerformanceTrendsTypes";
+
+/**
+ * Memory statistics metric data structure.
+ */
+export class MemoryStatsData extends PerformanceTrendContext implements ISummaryMetric {
+
+    // Properties must match server-side MemoryStatsTelemetryRecord
+    heapUsedMb: number;
+    textureMemoryMb: number;
+    meshMemoryMb: number;
+    audioMemoryMb: number;
+    totalAllocatedMb: number;
+
+    /**
+     * The metric type identifier. Must match server-side PerformanceSummaryType.
+     */
+    static readonly metricType = "MemoryStats";
+
+    constructor() {
+        super();
+    }
+}
+
+// Register with the metric type registry
+MetricTypeRegistry.register(MemoryStatsData, new MemoryStatsViewGenerator());
+```
+
+> **Important**: The `static readonly metricType` must exactly match the server-side `PerformanceSummaryType`.
+
+### Step 5: Create the View Generator
+
+Create in `PerformanceTrends/viewgenerators/`:
+
+```typescript
+// MemoryStatsViewGenerator.tsx
+
+import { observer } from "mobx-react-lite";
+import { graphColors } from "hordePlugins/analytics/telemetryData";
+import { LineGraph } from "../components/LineGraphComponent";
+import { MemoryStatsData } from "../metrictypes/MemoryStatsData";
+import { IMetricView, IMetricViewGenerator } from "./PerformanceTrendRenderTypes";
+import { PerformanceTrendUIOptionsState } from "../filters/PerformanceTrendUIOptionsState";
+
+/**
+ * View generator for memory statistics metrics.
+ */
+export class MemoryStatsViewGenerator implements IMetricViewGenerator<MemoryStatsData> {
+
+    getViews(metrics: MemoryStatsData[], uiState: PerformanceTrendUIOptionsState): IMetricView[] {
+
+        // Extract data series
+        const heapUsed = metrics.map(m => m.heapUsedMb);
+        const textureMemory = metrics.map(m => m.textureMemoryMb);
+        const meshMemory = metrics.map(m => m.meshMemoryMb);
+        const totalAllocated = metrics.map(m => m.totalAllocatedMb);
+
+        // Calculate averages for display
+        const avgHeap = heapUsed.length > 0
+            ? heapUsed.reduce((a, b) => a + b, 0) / heapUsed.length
+            : 0;
+
+        // Define views
+        const HeapMemoryView = observer(({ width, height }: { width: number; height: number }) => {
+            const { curveGraphLine } = uiState;
+
+            return (
+                <div>
+                    <div>
+                        <span>Heap Memory Usage</span><br />
+                        <span>(Avg) {avgHeap.toFixed(1)} MB</span>
+                    </div>
+                    <LineGraph
+                        data={[heapUsed]}
+                        legendLabels={["Heap Used"]}
+                        colorLabels={[graphColors[0]]}
+                        xAxisLabel="Time/Index"
+                        yAxisLabel="Memory (MB)"
+                        width={width}
+                        height={height}
+                        yAxisZeroScale={true}
+                        applyCurveSmoothing={curveGraphLine}
+                    />
+                </div>
+            );
+        });
+
+        const MemoryBreakdownView = observer(({ width, height }: { width: number; height: number }) => {
+            const { curveGraphLine } = uiState;
+
+            return (
+                <div>
+                    <div>
+                        <span>Memory Breakdown by Type</span>
+                    </div>
+                    <LineGraph
+                        data={[textureMemory, meshMemory, heapUsed]}
+                        legendLabels={["Textures", "Meshes", "Heap"]}
+                        xAxisLabel="Time/Index"
+                        yAxisLabel="Memory (MB)"
+                        width={width}
+                        height={height}
+                        yAxisZeroScale={true}
+                        applyCurveSmoothing={curveGraphLine}
+                    />
+                </div>
+            );
+        });
+
+        const TotalAllocationView = observer(({ width, height }: { width: number; height: number }) => {
+            const { curveGraphLine } = uiState;
+
+            return (
+                <div>
+                    <div>
+                        <span>Total Memory Allocated</span>
+                    </div>
+                    <LineGraph
+                        data={[totalAllocated]}
+                        legendLabels={["Total Allocated"]}
+                        colorLabels={[graphColors[5]]}
+                        xAxisLabel="Time/Index"
+                        yAxisLabel="Memory (MB)"
+                        width={width}
+                        height={height}
+                        yAxisZeroScale={true}
+                        applyCurveSmoothing={curveGraphLine}
+                    />
+                </div>
+            );
+        });
+
+        return [
+            { label: "Heap Memory", Render: HeapMemoryView },
+            { label: "Memory Breakdown", Render: MemoryBreakdownView },
+            { label: "Total Allocation", Render: TotalAllocationView },
+        ];
+    }
+}
+```
+
+### Step 6: Import for Side-Effect Registration
+
+Ensure the metric data file is imported so registration runs:
+
+```typescript
+// In PerformanceTrendView.tsx or a central imports file
+import "./PerformanceTrends/metrictypes/MemoryStatsData";  // Side-effect: registers metric
+```
+
+---
+
+## Architecture Summary
+
+```
+┌─────────────────────────────────────────────────────────────────┐
+│                         SERVER                                  │
+├─────────────────────────────────────────────────────────────────┤
+│  [AnalyticsTableGen]              → Generates column constants  │
+│  MemoryStatsTelemetryRecord       → Data schema                 │
+│  MemoryStatsPerformanceSummaryHandler → Query logic             │
+│  DI Registration                  → Service discovery           │
+└─────────────────────────────────────────────────────────────────┘
+                              │
+                              │  metricType = "MemoryStats"
+                              ▼
+┌─────────────────────────────────────────────────────────────────┐
+│                         CLIENT                                  │
+├─────────────────────────────────────────────────────────────────┤
+│  MemoryStatsData                  → TypeScript model            │
+│  MemoryStatsViewGenerator         → Visualization components    │
+│  MetricTypeRegistry.register()    → Single registration call    │
+└─────────────────────────────────────────────────────────────────┘
+```
+
+---
+
+## Checklist
+
+### Server-Side
+
+- [ ] Create telemetry record with `[AnalyticsTableGen]` attribute
+- [ ] Define `[Table]` and `[Column]` attributes for ORM mapping
+- [ ] Create handler extending `AbstractPerformanceSummaryTrendHandler`
+- [ ] Implement `PerformanceSummaryType` property (the type identifier)
+- [ ] Implement `ProcessGeneralMetricRequest` with custom query logic
+- [ ] Register Dapper column mapping in static constructor
+- [ ] Register handler in DI container
+
+### Client-Side
+
+- [ ] Create data class extending `PerformanceTrendContext`
+- [ ] Add `static readonly metricType` matching server-side
+- [ ] Define properties matching server-side record fields
+- [ ] Create view generator implementing `IMetricViewGenerator<T>`
+- [ ] Return `IMetricView[]` from `getViews()` for each visualization
+- [ ] Call `MetricTypeRegistry.register()` at module load
+- [ ] Import the data file to trigger registration
+
+### Verification
+
+- [ ] `metricType` strings match exactly between server and client
+- [ ] Property names match between `TelemetryRecord` and `Data` class
+- [ ] Build passes on both server (`dotnet build`) and client (`npm run build`)
+- [ ] New metric type appears in dashboard dropdown
+- [ ] Data loads and visualizations render correctly
diff --git a/Engine/Source/Programs/Horde/HordeDashboard/plugins/performancetrends/filters/PerformanceTrendOptionsController.ts b/Engine/Source/Programs/Horde/HordeDashboard/plugins/performancetrends/filters/PerformanceTrendOptionsController.ts
new file mode 100644
--- /dev/null
+++ b/Engine/Source/Programs/Horde/HordeDashboard/plugins/performancetrends/filters/PerformanceTrendOptionsController.ts
@@ -0,0 +1,186 @@
+// Copyright Epic Games, Inc. All Rights Reserved.
+
+import { action, makeObservable, observable } from "mobx";
+import { PerformanceTrendOptionsState } from "./PerformanceTrendOptionsState";
+import { GroupByOption, PerformanceTrendUIOptionsState } from "./PerformanceTrendUIOptionsState";
+import { PerformanceTrendOptionsDataHandler } from "./PerformanceTrendOptionsDataHandler";
+import { IPerformanceTrendOptionWriter as IPerformanceTrendOptionWriter, PerformanceTrendOptionsStateDiff, PerformanceTrendsOptionsWriter as PerformanceTrendOptionsWriter } from "./PerformanceTrendOptionsWriter";
+
+/**
+ * Primary class used to interact with the performance trend option state.
+ * @todo This is a code-sharing candidate between Performance Trends and Build Health.
+ */
+export class PerformanceTrendOptionsController {
+    // #region -- Private Members --
+
+    readonly state: PerformanceTrendOptionsState;
+    readonly uiState: PerformanceTrendUIOptionsState;
+    private writer: IPerformanceTrendOptionWriter | null;
+    private lastSynchronize = -1;
+
+    // #endregion -- Private Members --
+
+    // #region -- Public Members --
+
+    @observable optionsChangeVersion = 0;
+    @observable uiOptionsChangeVersion = 0;
+    querySchemaVersion: number = 1;
+
+    // #endregion -- Public Members --
+
+    // #region -- Constructor --
+
+    constructor(state: PerformanceTrendOptionsState, uiState: PerformanceTrendUIOptionsState) {
+        this.state = state;
+        this.uiState = uiState;
+
+        makeObservable(this);
+    }
+
+    // #endregion -- Constructor --
+
+    // #region -- Public Api --
+
+    /**
+     * Sets whether to use curved graph lines, or not.
+     * @param isActive True to use curved graph lines, false otherwise.
+     */
+    @action
+    setCurveGraphLine(isActive: boolean) {
+        this.uiState.curveGraphLine = isActive;
+        this.setUIOptionsChanged();
+    }
+
+    /**
+     * Sets the group by option for graph series.
+     * @param option The group by option to use.
+     */
+    @action
+    setGroupByOption(option: GroupByOption) {
+        this.uiState.groupByOption = option;
+        this.setUIOptionsChanged();
+    }
+
+    /**
+     * Toggles a viewable property on or off.
+     * @param key The composite key "{testType}::{propertyKey}".
+     * @param propertyKey The property key for display.
+     * @param isSelected Whether to enable or disable.
+     */
+    @action
+    toggleViewableProperty(key: string, propertyKey: string, isSelected: boolean) {
+        if (isSelected) {
+            this.uiState.enabledViewableProperties[key] = propertyKey;
+        } else {
+            delete this.uiState.enabledViewableProperties[key];
+        }
+        this.setUIOptionsChanged();
+    }
+
+    /**
+     * Toggles all provided viewable properties.
+     * @param properties Array of { key, propertyKey } to toggle.
+     * @param isSelected Whether to enable or disable.
+     */
+    @action
+    toggleAllViewableProperties(properties: { key: string, propertyKey: string }[], isSelected: boolean) {
+        for (const prop of properties) {
+            if (isSelected) {
+                this.uiState.enabledViewableProperties[prop.key] = prop.propertyKey;
+            } else {
+                delete this.uiState.enabledViewableProperties[prop.key];
+            }
+        }
+        this.setUIOptionsChanged();
+    }
+
+    /**
+     * Initializes default viewable properties for a test type.
+     * @param testType The test type.
+     * @param defaultProperties Array of property keys that are visible by default.
+     */
+    @action
+    initializeDefaultViewableProperties(testType: string, defaultProperties: string[]) {
+        const prefix = `${testType}::`;
+        // Only initialize if no properties are set for this test type
+        const hasExisting = Object.keys(this.uiState.enabledViewableProperties).some(k => k.startsWith(prefix));
+        if (!hasExisting) {
+            for (const prop of defaultProperties) {
+                this.uiState.enabledViewableProperties[`${prefix}${prop}`] = prop;
+            }
+            this.setUIOptionsChanged();
+        }
+    }
+
+    /**
+     * Starts a transaction session for updating the performance trend options.
+     * @returns 
+     */
+    getTransactionSession(): IPerformanceTrendOptionWriter {
+        if (!this.writer) {
+            this.writer = new PerformanceTrendOptionsWriter(this.state, this.uiState);
+        }
+
+        return this.writer;
+    }
+
+    /**
+     * Completes a transaction session for updating the performance trend options.
+     * @returns 
+     */
+    commitTransactionSession(): PerformanceTrendOptionsStateDiff {
+        if (this.writer === null || this.writer === undefined) {
+            return {
+                added: undefined,
+                removed: undefined,
+            };
+        }
+
+        let finalReceipts = this.writer.produceReceiptDiff();
+        this.writer = null;
+
+        // Only request a synchronization if we have produced diffs in our receipt.
+        if (finalReceipts.hasDiff) {
+            this.setOptionsChanged();
+            this.synchronizeDerivedKeys();
+        }
+
+        return finalReceipts;
+    }
+
+    // #endregion -- Public Api --
+
+    // #region -- Private Api --
+
+    /**
+     * Synchronize the options that are bound to UI elements, with the data observed by consumers of the option set.
+     * This is important to keep separate as it allows us to control when we signal/flush a finalized set of options to the consumers.
+     */
+    @action
+    private synchronizeDerivedKeys() {
+        if (this.lastSynchronize < this.optionsChangeVersion) {
+            this.state.publishedEnabledTemplateKeys = Object.keys(this.state.enabledTemplates);
+            this.state.publishedEnabledProjectKeys = Object.keys(this.state.enabledProjects);
+            this.state.publishedEnabledTestIdentityKeys = Object.keys(this.state.enabledTestIdentities);
+            this.state.publishedEnabledTestTypeKeys = Object.keys(this.state.enabledTestTypes);
+            this.state.publishedEnabledMetricSummaryTypes = Object.keys(this.state.enabledMetricSummaryTypes);
+            this.state.publishedEnabledStreamKeys = Object.keys(this.state.enabledStreams);
+            this.state.publishedEnabledPlatforms = Object.keys(this.state.enabledPlatforms);
+            this.state.publishedEnabledStartCommit = this.state.enabledStartCommit;
+            this.state.publishedEnabledEndCommit = this.state.enabledEndCommit;
+            this.lastSynchronize = this.optionsChangeVersion;
+        }
+    }
+
+    @action
+    private setOptionsChanged() {
+        this.optionsChangeVersion++;
+    }
+
+    @action
+    private setUIOptionsChanged() {
+        this.uiOptionsChangeVersion++;
+    }
+
+    // #endregion -- Private Api --
+}
\ No newline at end of file
diff --git a/Engine/Source/Programs/Horde/HordeDashboard/plugins/performancetrends/filters/PerformanceTrendOptionsDataHandler.ts b/Engine/Source/Programs/Horde/HordeDashboard/plugins/performancetrends/filters/PerformanceTrendOptionsDataHandler.ts
new file mode 100644
--- /dev/null
+++ b/Engine/Source/Programs/Horde/HordeDashboard/plugins/performancetrends/filters/PerformanceTrendOptionsDataHandler.ts
@@ -0,0 +1,511 @@
+// Copyright Epic Games, Inc. All Rights Reserved.
+
+import { action, makeObservable, observable, runInAction } from "mobx";
+import { PerformanceTrendOptionsState } from "./PerformanceTrendOptionsState";
+import { PerformanceTrendOptionsController } from "./PerformanceTrendOptionsController";
+import { GetTestProjectCommitResponse, GetTestProjectPlatformResponse, GetTestProjectResponse } from "../responses/GetTestProjectResponse";
+import { decodeTestTypeKey, decodeMetricTypeKey, DecodedTestIdentityKey, decodeTestIdentity, decodeStreamKey, DecodedStreamKey, decodePlatformKey, DecodedPlatformKey } from "../responses/FilterKeys";
+import { getCommits, getPlatforms, getTestProjects } from "../api";
+
+// #region -- Data Refresh Request --
+
+/**
+ * Models a request to the data handler.
+ */
+export type DataHandlerRefreshRequest = {
+    projects?: boolean;
+    testIdentities?: boolean;
+    testTypes?: boolean;
+    metricTypes?: boolean;
+    streams?: boolean;
+    platforms?: boolean;
+    commits?: boolean;
+}
+
+export const EMPTY_REFRESH_REQUEST: DataHandlerRefreshRequest = { projects: false, testIdentities: false, testTypes: false, metricTypes: false, streams: false, platforms: false, commits: false };
+export const DEFAULT_REFRESH_REQUEST: DataHandlerRefreshRequest = { projects: true, testIdentities: true, testTypes: true, metricTypes: true, streams: true, platforms: true, commits: true };
+
+/**
+ * Utility function that will expand a Data Handler Refresh Request to include all the dependents in it's refresh.
+ * @param request The request to expand.
+ * @returns The expanded request.
+ */
+export function expandRefreshRequest(request: DataHandlerRefreshRequest): DataHandlerRefreshRequest {
+    const expanded: DataHandlerRefreshRequest = { ...request };
+
+    if (request.projects) {
+        expanded.testIdentities = true;
+    }
+
+    if (expanded.testIdentities) {
+        expanded.testTypes = true;
+    }
+
+    if (expanded.testTypes) {
+        expanded.metricTypes = true;
+    }
+
+    if (expanded.streams) {
+        expanded.metricTypes = true;
+    }
+
+    if (expanded.metricTypes) {
+        expanded.streams = true;
+    }
+
+    if (expanded.streams) {
+        expanded.platforms = true;
+    }
+
+    if (expanded.platforms) {
+        expanded.commits = true;
+    }
+
+    return expanded;
+}
+
+// #endregion -- Data Refresh Request --
+
+/**
+ * Data handler responsible for obtaining data used in options controls.
+ * @todo This is a code-sharing candidate between Performance Trends and Build Health.
+ */
+export class PerformanceTrendOptionsDataHandler {
+    // #region -- Private Members --
+
+    private lastRefresh = -1;
+    private performanceTrendOptionsState: PerformanceTrendOptionsState;
+    private performanceTrendOptionsController: PerformanceTrendOptionsController;
+
+    // #endregion -- Private Members --
+
+    // #region -- Public Members --
+
+    @observable
+    isLoadingProjects: boolean = false;
+
+    @observable
+    isLoadingPlatforms: boolean = false;
+
+    @observable
+    isLoadingCommits: boolean = false;
+
+    @observable.shallow
+    projectsData: GetTestProjectResponse[] = [];
+
+    @observable.shallow
+    testIdentities: GetTestProjectResponse[] = [];
+
+    @observable.shallow
+    testTypes: GetTestProjectResponse[] = [];
+
+    @observable.shallow
+    metricSummaryTypes: GetTestProjectResponse[] = [];
+
+    @observable.shallow
+    streams: GetTestProjectResponse[] = [];
+
+    @observable.shallow
+    platforms: GetTestProjectPlatformResponse[] = [];
+
+    @observable.shallow
+    commits: GetTestProjectCommitResponse[] = [];
+
+    // #endregion -- Public Members --
+
+    // #region -- Constructor --
+
+    constructor(performanceTrendOptionsState: PerformanceTrendOptionsState, performanceTrendOptionsController: PerformanceTrendOptionsController) {
+        makeObservable(this);
+
+        this.performanceTrendOptionsState = performanceTrendOptionsState;
+        this.performanceTrendOptionsController = performanceTrendOptionsController;
+    }
+
+    // #endregion -- Constructor --
+
+    // #region -- Private API --
+
+    private async refreshData(request: DataHandlerRefreshRequest) {
+        return this.getProjectsData(request);
+    }
+
+    // #endregion -- Private API --
+
+    // #region -- Public API --
+
+    /**
+     * Requests a full option hierarchy refresh.
+     * @param request The request which dictates the degree of refresh.
+     * @param force Whether to force the operation or not.
+     * @param onComplete The delegate to invoke upon completion of data refresh.
+     */
+    async requestHierarchicalRefresh(request: DataHandlerRefreshRequest = DEFAULT_REFRESH_REQUEST, force?: boolean, onComplete?: (request: DataHandlerRefreshRequest) => void) {
+        if (this.lastRefresh < this.performanceTrendOptionsController.optionsChangeVersion || force) {
+            this.lastRefresh = this.performanceTrendOptionsController.optionsChangeVersion;
+            await this.refreshData(request);
+            onComplete?.(request);
+        }
+    }
+
+    /**
+     * Resets the handler, returning it to a project-only data state.
+     */
+    @action
+    reset() {
+        this.projectsData = [];
+        this.testIdentities = [];
+        this.metricSummaryTypes = [];
+        this.testTypes = [];
+        this.streams = [];
+        this.platforms = [];
+        this.commits = [];
+    }
+
+    /**
+     * Hierarchical reset of the options.
+     * @param request The request to use in processing the reset.
+     */
+    @action
+    hierarchyReset(request: DataHandlerRefreshRequest) {
+        let cascade = false;
+
+        if (request.projects || cascade) {
+            this.projectsData = [];
+            cascade = true;
+        }
+        if (request.testIdentities || cascade) {
+            this.testIdentities = [];
+            cascade = true;
+        }
+        if (request.metricTypes || cascade) {
+            this.metricSummaryTypes = [];
+            cascade = true;
+        }
+        if (request.testTypes || cascade) {
+            this.testTypes = [];
+            cascade = true;
+        }
+        if (request.streams || cascade) {
+            this.streams = [];
+            cascade = true;
+        }
+        if (request.platforms || cascade) {
+            this.platforms = [];
+            cascade = true;
+        }
+        if (request.commits || cascade) {
+            this.commits = [];
+        }
+    }
+
+    /**
+     * Obtains the test project data given the current performance trend options.
+     * @param request The data refresh request message to use in order to control refresh depth.
+     */
+    async getProjectsData(request: DataHandlerRefreshRequest) {
+        if (this.projectsData.length === 0 || request.projects) {
+            runInAction(() => {
+                this.isLoadingProjects = true;
+            });
+
+            const projectDataResult = (await getTestProjects());
+
+            runInAction(() => {
+                this.projectsData = projectDataResult;
+                this.isLoadingProjects = false;
+            });
+        }
+
+        return this.getTestIdentityData(request);
+    }
+
+    /**
+     * Obtains the test identity data given the currently enabled test projects.
+     * @param request The refresh request to cascade to downstream metric types.
+     */
+    async getTestIdentityData(request: DataHandlerRefreshRequest) {
+        if (request.testIdentities && Object.keys(this.performanceTrendOptionsState.enabledProjects).length > 0) {
+            runInAction(() => {
+                this.testIdentities = [];
+            });
+
+            // For all of our enabled projects, add test identities that belong to them to the selection set.
+            for (const key in this.performanceTrendOptionsState.enabledProjects) {
+                let testIdentityData: GetTestProjectResponse[] = [];
+                let testIdentityTargetSet: Set<string> = new Set<string>();
+
+                // Make sure we are inserting test identities that falls under the current enabled project selection
+                this.projectsData.forEach((value: GetTestProjectResponse) => {
+                    if (key === value.testName && !testIdentityTargetSet.has(value.testIdentity)) {
+                        testIdentityData.push(value);
+                        testIdentityTargetSet.add(value.testIdentity);
+                    }
+                });
+                runInAction(() => {
+                    this.testIdentities.push(...testIdentityData);
+                });
+            }
+        }
+
+        return this.getMetricTypeData(request);
+    }
+
+    /**
+    * Obtains the metric types given the currently enabled test projects, and test identities.
+    * @param request The refresh request to cascade to downstream streams.
+    */
+    async getMetricTypeData(request: DataHandlerRefreshRequest) {
+        if (request.metricTypes && Object.keys(this.performanceTrendOptionsState.enabledTestIdentities).length > 0) {
+            runInAction(() => {
+                this.metricSummaryTypes = [];
+            });
+
+            // For all of our enabled test types, add metric summary types that belong to them.
+            for (const key in this.performanceTrendOptionsState.enabledTestIdentities) {
+                let metricTypeData: GetTestProjectResponse[] = [];
+                let metricTypeTargetSet: Set<string> = new Set<string>();
+                let decodedKey = decodeTestIdentity(key);
+
+                // Make sure we are inserting metric data that falls under the current enabled project selection
+                this.projectsData.forEach((value: GetTestProjectResponse) => {
+                    if (decodedKey?.testProject === value.testName && decodedKey?.testIdentity === value.testIdentity && !metricTypeTargetSet.has(value.summaryType)) {
+                        metricTypeData.push(value);
+                        metricTypeTargetSet.add(value.summaryType);
+                    }
+                });
+
+                runInAction(() => {
+                    this.metricSummaryTypes = metricTypeData;
+                });
+            }
+        }
+
+        return this.getTestTypeData(request);
+    }
+
+    /**
+    * Obtains the test identity data given the currently enabled test projects.
+    * @param request The refresh request to cascade to downstream metric types.
+    */
+    async getTestTypeData(request: DataHandlerRefreshRequest) {
+        if (request.testTypes && Object.keys(this.performanceTrendOptionsState.enabledMetricSummaryTypes).length > 0) {
+            runInAction(() => {
+                this.testTypes = [];
+            });
+
+            // For all of our enabled projects, add test types that belong to them to the selection set.
+            for (const key in this.performanceTrendOptionsState.enabledMetricSummaryTypes) {
+                let testTypeData: GetTestProjectResponse[] = [];
+                let testTypeTargetSet: Set<string> = new Set<string>();
+                let decodedMetricSummaryType = decodeMetricTypeKey(key);
+
+                // Make sure we are inserting test types that falls under the current enabled project selection
+                this.projectsData.forEach((value: GetTestProjectResponse) => {
+                    if (decodedMetricSummaryType?.testProject === value.testName && decodedMetricSummaryType?.testIdentity === value.testIdentity && decodedMetricSummaryType?.summaryType === value.summaryType && !testTypeTargetSet.has(value.testType)) {
+                        testTypeData.push(value);
+                        testTypeTargetSet.add(value.testType);
+                    }
+                });
+                runInAction(() => {
+                    this.testTypes.push(...testTypeData);
+                });
+            }
+        }
+
+        return this.getStreams(request);
+    }
+
+    /**
+    * Obtains the streams given the currently enabled test projects, test identities, and metric types.
+    * @param request The refresh request to cascade to downstream platforms.
+    */
+    async getStreams(request: DataHandlerRefreshRequest) {
+        if (request.streams && Object.keys(this.performanceTrendOptionsState.enabledTestTypes).length > 0) {
+            runInAction(() => {
+                this.streams = [];
+            });
+
+            let streamData: GetTestProjectResponse[] = [];
+
+            // For all of our enabled test types, add streams that belong to them (and the other related ancestors).
+            for (const key in this.performanceTrendOptionsState.enabledTestTypes) {
+                let singleStreamData: GetTestProjectResponse[] = [];
+                const streamTargetSet: Set<string> = new Set<string>();
+                const decodedKey = decodeTestTypeKey(key);
+
+                // Make sure we are inserting metric data that falls under the current enabled project selection
+                this.projectsData.forEach((value: GetTestProjectResponse) => {
+                    if (decodedKey?.metricSummaryType === value.summaryType && decodedKey?.testProject === value.testName && decodedKey?.testIdentity === value.testIdentity && decodedKey?.testType === value.testType && !streamTargetSet.has(value.computedStream)) {
+                        singleStreamData.push(value);
+                        streamTargetSet.add(value.computedStream);
+                    }
+                });
+
+                streamData.push(...singleStreamData);
+            }
+
+            runInAction(() => {
+                this.streams = streamData;
+            });
+        }
+
+        return this.getPlatforms(request);
+    }
+
+    /**
+    * Obtains the platforms with valid results, given the currently enabled test projects, test identities, and summary type.
+    * @param request The refresh request to cascade to downstream metric types.
+    */
+    async getPlatforms(request: DataHandlerRefreshRequest) {
+        if (request.platforms && Object.keys(this.performanceTrendOptionsState.enabledStreams).length > 0) {
+
+            runInAction(() => {
+                this.platforms = [];
+            });
+
+            const publishedEnabledProjectKeys = Object.keys(this.performanceTrendOptionsState.enabledProjects);
+            const publishedEnabledTestIdentityKeys = Object.keys(this.performanceTrendOptionsState.enabledTestIdentities);
+            const publishedEnabledMetricSummaryTypes = Object.keys(this.performanceTrendOptionsState.enabledMetricSummaryTypes);
+            const publishedEnabledStreamKeys = Object.keys(this.performanceTrendOptionsState.enabledStreams);
+
+            const testProject = publishedEnabledProjectKeys[0];
+
+            const decodedTestIdentity: DecodedTestIdentityKey | null = decodeTestIdentity(publishedEnabledTestIdentityKeys[0]);
+
+            if (!decodedTestIdentity) {
+                console.warn("[GetPlatforms] Was unable to obtain a test identity from the enabled test identity key.");
+
+                this.hierarchyReset(expandRefreshRequest({ commits: true }));
+
+                return;
+            }
+
+            let testIdentity = decodedTestIdentity.testIdentity;
+
+            const decodedMetricSummaryType = decodeMetricTypeKey(publishedEnabledMetricSummaryTypes[0]);
+
+            if (!decodedMetricSummaryType) {
+                console.warn("[GetPlatforms] Was unable to obtain a metric summary type from the enabled metric summary type.");
+
+                this.hierarchyReset(expandRefreshRequest({ commits: true }));
+
+                return;
+            }
+
+            let metricSummaryType = decodedMetricSummaryType.summaryType;
+
+            const decodedStreamKeys = publishedEnabledStreamKeys
+                .map(x => decodeStreamKey(x))
+                .filter((x: DecodedStreamKey | null): x is DecodedStreamKey => x !== null);
+
+            if (decodedStreamKeys.length === 0) {
+                console.warn("[GetPlatforms] Was unable to obtain a stream key from the enabled stream.");
+
+                this.hierarchyReset(expandRefreshRequest({ commits: true }));
+
+                return;
+            }
+
+            let streams = decodedStreamKeys.map(x => x.stream);
+
+            runInAction(() => {
+                this.isLoadingPlatforms = true;
+            });
+
+            const platformResult = (await getPlatforms(testProject, testIdentity, metricSummaryType, streams));
+
+            runInAction(() => {
+                this.isLoadingPlatforms = false;
+                this.platforms = platformResult;
+            });
+        } else if (request.platforms) {
+            runInAction(() => {
+                this.platforms = [];
+            });
+        }
+
+        return this.getCommits(request);
+    }
+
+    /**
+    * Obtains the commits with valid results, given the currently enabled test projects, test identities, summary type, and platform.
+    * @param request The refresh request to cascade to downstream metric types.
+    */
+    async getCommits(request: DataHandlerRefreshRequest) {
+        if (request.commits && Object.keys(this.performanceTrendOptionsState.enabledPlatforms).length > 0) {
+            runInAction(() => {
+                this.commits = [];
+            });
+
+            const publishedEnabledProjectKeys = Object.keys(this.performanceTrendOptionsState.enabledProjects);
+            const publishedEnabledTestIdentityKeys = Object.keys(this.performanceTrendOptionsState.enabledTestIdentities);
+            const publishedEnabledMetricSummaryTypes = Object.keys(this.performanceTrendOptionsState.enabledMetricSummaryTypes);
+            const publishedEnabledStreamKeys = Object.keys(this.performanceTrendOptionsState.enabledStreams);
+            const publishedPlatformKeys = Object.keys(this.performanceTrendOptionsState.enabledPlatforms);
+
+            const testProject = publishedEnabledProjectKeys[0];
+
+            const decodedTestIdentity: DecodedTestIdentityKey | null = decodeTestIdentity(publishedEnabledTestIdentityKeys[0]);
+
+            if (!decodedTestIdentity) {
+                console.warn("[GetCommits] Was unable to obtain a test identity from the enabled test identity key.");
+
+                return;
+            }
+
+            const testIdentity = decodedTestIdentity.testIdentity;
+
+            const decodedMetricSummaryType = decodeMetricTypeKey(publishedEnabledMetricSummaryTypes[0]);
+
+            if (!decodedMetricSummaryType) {
+                console.warn("[GetCommits] Was unable to obtain a metric summary type from the enabled metric summary type.");
+
+                return;
+            }
+
+            const metricSummaryType = decodedMetricSummaryType.summaryType;
+
+            const decodedStreamKeys = publishedEnabledStreamKeys
+                .map(x => decodeStreamKey(x))
+                .filter((x: DecodedStreamKey | null): x is DecodedStreamKey => x !== null);
+
+            if (decodedStreamKeys.length === 0) {
+                console.warn("[GetCommits] Was unable to obtain a stream key from the enabled stream.");
+
+                return;
+            }
+
+            const streams = decodedStreamKeys.map(x => x.stream);
+
+            const decodedPlatformKeys = publishedPlatformKeys
+                .map(x => decodePlatformKey(x))
+                .filter((x: DecodedPlatformKey | null): x is DecodedPlatformKey => x !== null);
+
+            if (decodedPlatformKeys.length === 0) {
+                console.warn("[GetCommits] Was unable to obtain a platform key from the enabled platform.");
+
+                return;
+            }
+
+            const platforms = decodedPlatformKeys.map(x => x.platform);
+
+            runInAction(() => {
+                this.isLoadingCommits = true;
+            });
+
+            const commitsResult = (await getCommits(testProject, testIdentity, metricSummaryType, streams, platforms));
+
+            runInAction(() => {
+                this.isLoadingCommits = false;
+                this.commits = commitsResult;
+            });
+        } else if (request.commits) {
+            runInAction(() => {
+                this.commits = [];
+            });
+        }
+    }
+
+    // #endregion -- Public Api --
+}
\ No newline at end of file
diff --git a/Engine/Source/Programs/Horde/HordeDashboard/plugins/performancetrends/filters/PerformanceTrendOptionsState.ts b/Engine/Source/Programs/Horde/HordeDashboard/plugins/performancetrends/filters/PerformanceTrendOptionsState.ts
new file mode 100644
--- /dev/null
+++ b/Engine/Source/Programs/Horde/HordeDashboard/plugins/performancetrends/filters/PerformanceTrendOptionsState.ts
@@ -0,0 +1,71 @@
+// Copyright Epic Games, Inc. All Rights Reserved.
+
+import { makeObservable, observable } from "mobx";
+import { PerformanceTrendOptionStateReceipt } from "./PerformanceTrendOptionsWriter";
+
+/**
+ * * The state object for query options on the Performance Trend dashboard.
+ * @todo This is a code-sharing candidate between Performance Trends and Build Health.
+ */
+export class PerformanceTrendOptionsState {
+
+    // #region -- Published Options --
+
+    @observable publishedEnabledProjectKeys: string[] = [];
+    @observable publishedEnabledTestIdentityKeys: string[] = [];
+    @observable publishedEnabledTestTypeKeys: string[] = [];
+    @observable publishedEnabledMetricSummaryTypes: string[] = [];
+    @observable publishedEnabledStartCommit: number | null = null;
+    @observable publishedEnabledEndCommit: number | null = null;
+    @observable publishedEnabledPlatforms: string[] = [];
+    @observable publishedEnabledStreamKeys: string[] = [];
+    @observable publishedEnabledTemplateKeys: string[] = [];
+
+    // #endregion -- Published Options --
+
+    // #region -- User Selected Options --
+
+    @observable enabledProjects: Record<string, string> = {};
+    @observable enabledTestIdentities: Record<string, string> = {};
+    @observable enabledTestTypes: Record<string, string> = {};
+    @observable enabledMetricSummaryTypes: Record<string, string> = {};
+    @observable enabledPlatforms: Record<string, string> = {};
+    @observable enabledStartCommit: number | null = null;
+    @observable enabledEndCommit: number| null = null;
+    @observable enabledStreams: Record<string, string> = {};
+    @observable enabledTemplates: Record<string, string> = {};
+
+    // #endregion -- User Selected Options --
+
+    // #region -- Constructor --
+
+    constructor() {
+        makeObservable(this);
+    }
+
+    // #endregion -- Constructor --
+
+    // #region -- Public Api --
+
+    /**
+     * Generates a state receipt for the options in it's current state.
+     * @returns A receipt representation of the current performance trend options.
+     */
+    generateStateReceipt(): PerformanceTrendOptionStateReceipt {
+        let receipt: PerformanceTrendOptionStateReceipt = new PerformanceTrendOptionStateReceipt(
+            Object.keys(this.enabledProjects),
+            Object.keys(this.enabledTestIdentities),
+            Object.keys(this.enabledTestTypes),
+            Object.keys(this.enabledMetricSummaryTypes),
+            Object.keys(this.enabledPlatforms),
+            this.enabledStartCommit,
+            this.enabledEndCommit,
+            Object.keys(this.enabledStreams),
+            Object.keys(this.enabledTemplates)
+        );
+
+        return receipt;
+    }
+
+    // #endregion -- Public Api --
+}
\ No newline at end of file
diff --git a/Engine/Source/Programs/Horde/HordeDashboard/plugins/performancetrends/filters/PerformanceTrendOptionsWriter.ts b/Engine/Source/Programs/Horde/HordeDashboard/plugins/performancetrends/filters/PerformanceTrendOptionsWriter.ts
new file mode 100644
--- /dev/null
+++ b/Engine/Source/Programs/Horde/HordeDashboard/plugins/performancetrends/filters/PerformanceTrendOptionsWriter.ts
@@ -0,0 +1,504 @@
+// Copyright Epic Games, Inc. All Rights Reserved.
+
+import { action, makeObservable } from "mobx";
+import { PerformanceTrendOptionsState } from "./PerformanceTrendOptionsState";
+import { PerformanceTrendUIOptionsState } from "./PerformanceTrendUIOptionsState";
+
+/**
+ * Interface that describes the difference between two option states.
+ */
+export interface PerformanceTrendOptionsStateDiff {
+    added?: PerformanceTrendOptionStateReceipt;
+    removed?: PerformanceTrendOptionStateReceipt;
+    hasDiff?: boolean;
+}
+
+/**
+ * Performance Trend Option state receipt.
+ */
+export class PerformanceTrendOptionStateReceipt {
+    enabledProjects: ReadonlySet<string>;
+    enabledTestIdentities: ReadonlySet<string>;
+    enabledTestTypes: ReadonlySet<string>;
+    enabledMetricSummaries: ReadonlySet<string>;
+    enabledPlatforms: ReadonlySet<string>;
+    enabledStartCommit: number | null;
+    enabledEndCommit: number | null;
+    enabledStreams: ReadonlySet<string>;
+    enabledTemplates: ReadonlySet<string>;
+
+    constructor(
+        projects: Iterable<string> = [],
+        testIdentities: Iterable<string> = [],
+        testTypes: Iterable<string> = [],
+        metricSummaries: Iterable<string> = [],
+        enabledPlatforms: Iterable<string> = [],
+        enabledStartCommit: number | null,
+        enabledEndCommit: number | null,
+        streams: Iterable<string> = [],
+        templates: Iterable<string> = [],
+
+    ) {
+        this.enabledProjects = new Set(projects);
+        this.enabledTestIdentities = new Set(testIdentities);
+        this.enabledTestTypes = new Set(testTypes);
+        this.enabledMetricSummaries = new Set(metricSummaries);
+        this.enabledPlatforms = new Set(enabledPlatforms);
+        this.enabledStartCommit = enabledStartCommit;
+        this.enabledEndCommit = enabledEndCommit;
+        this.enabledStreams = new Set(streams);
+        this.enabledTemplates = new Set(templates);
+    }
+
+    /**
+        * Produces a build health option state diff between two provided receipts.
+        * @param prev The previous option state.
+        * @param next The next option state.
+        * @returns The diff between the two receipts.
+        */
+    static diff(
+        prev: PerformanceTrendOptionStateReceipt,
+        next: PerformanceTrendOptionStateReceipt
+    ): PerformanceTrendOptionsStateDiff {
+
+        const added = new PerformanceTrendOptionStateReceipt(
+            PerformanceTrendOptionStateReceipt.difference(next.enabledProjects, prev.enabledProjects),
+            PerformanceTrendOptionStateReceipt.difference(next.enabledTestIdentities, prev.enabledTestIdentities),
+            PerformanceTrendOptionStateReceipt.difference(next.enabledTestTypes, prev.enabledTestTypes),
+            PerformanceTrendOptionStateReceipt.difference(next.enabledMetricSummaries, prev.enabledMetricSummaries),
+            PerformanceTrendOptionStateReceipt.difference(next.enabledPlatforms, prev.enabledPlatforms),
+            next.enabledStartCommit,
+            next.enabledEndCommit,
+            PerformanceTrendOptionStateReceipt.difference(next.enabledStreams, prev.enabledStreams),
+            PerformanceTrendOptionStateReceipt.difference(next.enabledTemplates, prev.enabledTemplates),
+        );
+
+        const removed = new PerformanceTrendOptionStateReceipt(
+            PerformanceTrendOptionStateReceipt.difference(prev.enabledProjects, next.enabledProjects),
+            PerformanceTrendOptionStateReceipt.difference(prev.enabledTestIdentities, next.enabledTestIdentities),
+            PerformanceTrendOptionStateReceipt.difference(prev.enabledTestTypes, next.enabledTestTypes),
+            PerformanceTrendOptionStateReceipt.difference(prev.enabledMetricSummaries, next.enabledMetricSummaries),
+            PerformanceTrendOptionStateReceipt.difference(prev.enabledPlatforms, next.enabledPlatforms),
+            prev.enabledStartCommit,
+            prev.enabledEndCommit,
+            PerformanceTrendOptionStateReceipt.difference(prev.enabledStreams, next.enabledStreams),
+            PerformanceTrendOptionStateReceipt.difference(prev.enabledTemplates, next.enabledTemplates),
+        );
+
+        let startCommitChanged: boolean = next.enabledStartCommit !== prev.enabledStartCommit;
+        let endCommitChanged: boolean = next.enabledEndCommit !== prev.enabledEndCommit;
+
+        const hasDiff = this.receiptHasChanges(added) || this.receiptHasChanges(removed) || startCommitChanged || endCommitChanged;
+
+        return { added, removed, hasDiff };
+    }
+
+    private static receiptHasChanges(r: PerformanceTrendOptionStateReceipt): boolean {
+        return (
+            r.enabledProjects.size > 0 ||
+            r.enabledTestIdentities.size > 0 ||
+            r.enabledTestTypes.size > 0 ||
+            r.enabledMetricSummaries.size > 0 ||
+            r.enabledPlatforms.size > 0 ||
+            r.enabledStreams.size > 0 ||
+            r.enabledTemplates.size > 0
+        );
+    }
+
+    private static difference<T>(a: ReadonlySet<T>, b: ReadonlySet<T>): Set<T> {
+        const out = new Set<T>();
+        for (const item of a) {
+            if (!b.has(item)) out.add(item);
+        }
+        return out;
+    }
+}
+
+/**
+ * Defines an interface to modify the performance trend options.
+ */
+export interface IPerformanceTrendOptionWriter {
+
+    /**
+     * Toggles a single test project to be used in the resulting view.
+     * @param key The key.
+     * @param testProjectName The name of the project.
+     */
+    toggleSingleProject(key: string, testProjectName: string): void;
+
+    /**
+     * Toggles a single test identity to be used in the resulting view. test identity belongs to a test project.
+     * @param key The key.
+     * @param testIdentity The name of the test identity.
+     */
+    toggleSingleTestIdentity(key: string, testIdentity: string): void;
+
+    /**
+     * Toggles a single metric summary type to be used in the resulting view. Metric summary type belongs to a test project & test identity.
+     * @param key The key.
+     * @param metricSummaryType The name of the metric summary type.
+     */
+    toggleSingleMetricSummaryType(key: string, metricSummaryType: string): void;
+
+    /**
+     * Toggles a single test type to be used in the resulting view. Test type belongs to a test project & test identity.
+     * @param key The key.
+     * @param testType The name of the test type.
+     */
+    toggleSingleTestType(key: string, testType: string): void;
+
+    /**
+     * Toggles a test type to be used in the resulting view. Test type belongs to a test project & test identity.
+     * @param key The key.
+     * @param testType The name of the test type.
+     */
+    toggleTestType(key: string, testType: string, isSelected: boolean): void;
+
+    /**
+     * Toggles all provided test types.
+     * @param testTypes The list of test types to toggle.
+     * @param isSelected Whether to enable the test type, or not.
+     */
+    toggleAllTestTypes(testTypes: { key: string, testType: string }[], isSelected: boolean): void;
+
+    /**
+     * Toggles a single stream to be used in the resulting view. Stream belongs to a test project, test identity, and metric summary type.
+     * @param key The key.
+     * @param stream The name of the stream.
+     */
+    toggleStream(key: string, stream: string, isSelected: boolean): void;
+
+    /**
+     * Toggles all provided streams.
+     * @param streams The list of streams to toggle.
+     * @param isSelected Whether to enable the stream, or not.
+     */
+    toggleAllStreams(streams: { key: string, stream: string }[], isSelected: boolean): void;
+
+    /**
+     * Toggles a single platform to be used in the resulting view. Platform belongs to a test project, test identity, and metric summary type.
+     * @param key The key.
+     * @param platform The name of the platform.
+     */
+    toggleSinglePlatform(key: string, platform: string): void;
+
+    /**
+     * Toggles a single platform to be used in the resulting view. Platform belongs to a test project, test identity, and metric summary type.
+     * @param key The key.
+     * @param platform The name of the platform.
+     */
+    togglePlatform(key: string, platform: string, isSelected: boolean): void;
+
+    /**
+     * Toggles all provided platforms.
+     * @param platforms The list of platforms to toggle.
+     * @param isSelected Whether to enable the platform, or not.
+     */
+    toggleAllPlatforms(platforms: { key: string, platform: string }[], isSelected: boolean): void;
+
+    /**
+     * Sets the start commit range to filter from.
+     * @param commitIdOrdered The ordered commit id to start filtering from.
+     */
+    toggleStartCommitRange(commitIdOrdered: number): void;
+
+    /**
+     * Sets the end commit range to filter to.
+     * @param commitIdOrdered The ordered commit id to end filtering on.
+     */
+    toggleEndCommitRange(commitIdOrdered: number): void;
+
+    /**
+     * Clears all the currently selected options.
+     */
+    clearAll(): void;
+
+    /**
+     * Clears all the currently selected test projects.
+     * @param cascade Whether to cascade the clear operation down to the dependency options.
+     */
+    clearTestProjects(cascade?: boolean): void;
+
+    /**
+     * Clears all the currently selected test identities.
+     * @param cascade Whether to cascade the clear operation down to the dependency options.
+     */
+    clearTestIdentities(cascade?: boolean): void;
+
+    /**
+     * Clears all the currently selected test types.
+     * @param cascade Whether to cascade the clear operation down to the dependency options.
+     */
+    clearTestTypes(cascade?: boolean): void;
+
+    /**
+     * Clears all the currently selected metric summary types.
+     * @param cascade Whether to cascade the clear operation down to the dependency options.
+     */
+    clearMetricSummaries(cascade?: boolean): void;
+
+    /**
+     * Clears all the currently selected streams.
+     * @param cascade Whether to cascade the clear operation down to the dependency options.
+     */
+    clearStreams(cascade?: boolean): void;
+
+    /**
+     * Clears all the currently selected platforms.
+     * @param cascade Whether to cascade the clear operation down to the dependency options.
+     */
+    clearPlatforms(cascade?: boolean): void;
+
+    /**
+     * Produces the set of additions and removals for all options.
+     */
+    produceReceiptDiff(): PerformanceTrendOptionsStateDiff;
+}
+
+/**
+ * Class that implements the performance trend option writer interface.
+ */
+export class PerformanceTrendsOptionsWriter implements IPerformanceTrendOptionWriter {
+
+    constructor(state: PerformanceTrendOptionsState, uiState: PerformanceTrendUIOptionsState) {
+        makeObservable(this);
+        this.state = state;
+        this.uiState = uiState;
+        this.activeTransactionReceipt = this.state.generateStateReceipt();
+    }
+
+    /**
+     * @inheritdoc 
+     */
+    @action
+    toggleSingleProject(key: string, testProjectName: string) {
+        this.state.enabledProjects = { [key]: testProjectName };
+        this.clearTestIdentities(true);
+    }
+
+    /**
+     * @inheritdoc 
+     */
+    @action
+    toggleSingleTestIdentity(key: string, testIdentity: string) {
+        this.state.enabledTestIdentities = { [key]: testIdentity };
+        this.clearMetricSummaries(true);
+    }
+
+    /**
+     * @inheritdoc 
+     */
+    @action
+    toggleSingleMetricSummaryType(key: string, metricSummaryType: string) {
+        this.state.enabledMetricSummaryTypes = { [key]: metricSummaryType };
+        this.clearTestTypes(true);
+    }
+
+    /**
+     * @inheritdoc 
+     */
+    @action
+    toggleSingleTestType(key: string, testType: string) {
+        this.state.enabledTestTypes = { [key]: testType };
+        this.clearStreams(true);
+    }
+
+    /**
+     * @inheritdoc 
+     */
+    @action
+    toggleTestType(key: string, testType: string, isSelected: boolean) {
+        if (isSelected) {
+            this.state.enabledTestTypes[key] = testType;
+        } else {
+            delete this.state.enabledTestTypes[key];
+            this.clearStreams(true);
+        }
+    }
+
+    /**
+     * @inheritdoc 
+     */
+    @action
+    toggleAllTestTypes(testTypes: { key: string; testType: string; }[], isSelected: boolean): void {
+        for (let idx: number = 0; idx < testTypes.length; ++idx) {
+            this.toggleTestType(testTypes[idx].key, testTypes[idx].testType, isSelected);
+        }
+    }
+
+    /**
+     * @inheritdoc 
+     */
+    @action
+    toggleStream(key: string, stream: string, enabled: boolean) {
+        if (enabled) {
+            this.state.enabledStreams[key] = stream;
+        } else {
+            delete this.state.enabledStreams[key];
+            this.clearPlatforms(true);
+        }
+    }
+
+    @action
+    toggleAllStreams(streams: { key: string, stream: string }[], isSelected: boolean) {
+        for (let idx: number = 0; idx < streams.length; ++idx) {
+            this.toggleStream(streams[idx].key, streams[idx].stream, isSelected);
+        }
+    }
+
+    /**
+     * @inheritdoc 
+     */
+    @action
+    toggleSinglePlatform(key: string, platform: string) {
+        this.state.enabledPlatforms = { [key]: platform };
+    }
+
+    /**
+     * @inheritdoc 
+     */
+    @action
+    togglePlatform(key: string, platform: string, isSelected: boolean) {
+        if (isSelected) {
+            this.state.enabledPlatforms[key] = platform;
+        } else {
+            delete this.state.enabledPlatforms[key];
+            this.clearCommits(true);
+        }
+    }
+
+    @action
+    toggleAllPlatforms(platforms: { key: string, platform: string }[], isSelected: boolean) {
+        for (let idx: number = 0; idx < platforms.length; ++idx) {
+            this.togglePlatform(platforms[idx].key, platforms[idx].platform, isSelected);
+        }
+    }
+
+    /**
+     * @inheritdoc 
+     */
+    @action
+    toggleStartCommitRange(commitIdOrdered: number) {
+        this.state.enabledStartCommit = commitIdOrdered;
+
+        // clamp higher if smaller
+        if (this.state.enabledEndCommit !== null && this.state.enabledStartCommit > this.state.enabledEndCommit) {
+            this.state.enabledEndCommit = commitIdOrdered;
+        }
+    }
+
+    /**
+     * @inheritdoc 
+     */
+    @action
+    toggleEndCommitRange(commitIdOrdered: number) {
+        this.state.enabledEndCommit = commitIdOrdered;
+
+        // clamp higher if smaller
+        if (this.state.enabledStartCommit !== null && this.state.enabledEndCommit < this.state.enabledStartCommit) {
+            this.state.enabledStartCommit = commitIdOrdered;
+        }
+    }
+
+    /**
+    * @inheritdoc
+    */
+    clearAll() {
+        this.clearTestProjects(true);
+    }
+
+    /**
+     * @inheritdoc
+     */
+    clearTestProjects(cascade?: boolean) {
+        Object.keys(this.state.enabledProjects).forEach(key => {
+            delete this.state.enabledProjects[key];
+        });
+
+        if (cascade) {
+            this.clearTestIdentities(cascade);
+        }
+    }
+
+    /**
+     * @inheritdoc
+     */
+    clearTestIdentities(cascade?: boolean) {
+        Object.keys(this.state.enabledTestIdentities).forEach(key => {
+            delete this.state.enabledTestIdentities[key];
+        });
+
+        if (cascade) {
+            this.clearMetricSummaries(cascade);
+        }
+    }
+
+    /**
+     * @inheritdoc
+     */
+    clearMetricSummaries(cascade?: boolean) {
+        Object.keys(this.state.enabledMetricSummaryTypes).forEach(key => {
+            delete this.state.enabledMetricSummaryTypes[key];
+        });
+
+        if (cascade) {
+            this.clearTestTypes(cascade);
+        }
+    }
+
+    /**
+     * @inheritdoc
+     */
+    clearTestTypes(cascade?: boolean) {
+        Object.keys(this.state.enabledTestTypes).forEach(key => {
+            delete this.state.enabledTestTypes[key];
+        });
+
+        if (cascade) {
+            this.clearStreams(cascade);
+        }
+    }
+
+    /**
+     * @inheritdoc
+     */
+    clearStreams(cascade?: boolean) {
+        Object.keys(this.state.enabledStreams).forEach(key => {
+            delete this.state.enabledStreams[key];
+        });
+
+        if (cascade) {
+            this.clearPlatforms(cascade);
+        }
+    }
+
+    /**
+    * @inheritdoc
+    */
+    clearPlatforms(cascade?: boolean) {
+        Object.keys(this.state.enabledPlatforms).forEach(key => {
+            delete this.state.enabledPlatforms[key];
+        });
+
+        this.clearCommits(cascade);
+    }
+
+    /**
+    * @inheritdoc
+    */
+    clearCommits(cascade?: boolean) {
+        this.state.enabledStartCommit = null;
+        this.state.enabledEndCommit = null;
+    }
+
+    produceReceiptDiff(): PerformanceTrendOptionsStateDiff {
+        let transactionStartReceipt = this.state.generateStateReceipt();
+        let returnSet = PerformanceTrendOptionStateReceipt.diff(this.activeTransactionReceipt, transactionStartReceipt);
+
+        return returnSet;
+    }
+
+    readonly state: PerformanceTrendOptionsState;
+    readonly uiState: PerformanceTrendUIOptionsState;
+    private activeTransactionReceipt: PerformanceTrendOptionStateReceipt;
+}
\ No newline at end of file
diff --git a/Engine/Source/Programs/Horde/HordeDashboard/plugins/performancetrends/filters/PerformanceTrendUIOptionsState.ts b/Engine/Source/Programs/Horde/HordeDashboard/plugins/performancetrends/filters/PerformanceTrendUIOptionsState.ts
new file mode 100644
--- /dev/null
+++ b/Engine/Source/Programs/Horde/HordeDashboard/plugins/performancetrends/filters/PerformanceTrendUIOptionsState.ts
@@ -0,0 +1,31 @@
+// Copyright Epic Games, Inc. All Rights Reserved.
+
+import { makeObservable, observable } from "mobx";
+
+export type GroupByOption =
+    | 'platform'
+    | 'stream'
+    | 'testType'
+    | 'platform+stream'
+    | 'platform+testType'
+    | 'stream+testType'
+    | 'platform+stream+testType';
+
+/**
+ * The state object for UI options on the Performance Trend dashboard.
+ * @todo This is a code-sharing candidate between Performance Trends and Build Health.
+ */
+export class PerformanceTrendUIOptionsState {
+    @observable curveGraphLine: boolean = false;
+    @observable groupByOption: GroupByOption = 'platform+stream';
+
+    /**
+     * Enabled viewable properties keyed by "{testType}::{propertyKey}".
+     * Value is the property key for display purposes.
+     */
+    @observable enabledViewableProperties: Record<string, string> = {};
+
+    constructor() {
+        makeObservable(this);
+    }
+}
\ No newline at end of file
diff --git a/Engine/Source/Programs/Horde/HordeDashboard/plugins/performancetrends/filters/PerformanceTrendUrlParams.ts b/Engine/Source/Programs/Horde/HordeDashboard/plugins/performancetrends/filters/PerformanceTrendUrlParams.ts
new file mode 100644
--- /dev/null
+++ b/Engine/Source/Programs/Horde/HordeDashboard/plugins/performancetrends/filters/PerformanceTrendUrlParams.ts
@@ -0,0 +1,29 @@
+// Copyright Epic Games, Inc. All Rights Reserved.
+
+import { GroupByOption } from "./PerformanceTrendUIOptionsState";
+
+/**
+ * Type that describes the url query params supported by the Performance Trends system that effect querying.
+ * @todo This is a code-sharing candidate between Performance Trends and Build Health.
+ */
+export type PerformanceTrendQueryParams = {
+    testProjects: string[];
+    testIdentities: string[];
+    testTypes: string[];
+    metricSummaries: string[];
+    streams: string[];
+    platforms: string[];
+    startCommit: number | null;
+    endCommit: number | null;
+    querySchemaVersion: number;
+};
+
+/**
+ * * Type that describes the url query params supported by the Performance Trends system that effect UI preferences.
+ * @todo This is a code-sharing candidate between Performance Trends and Build Health.
+ */
+export type PerformanceTrendUIParams = {
+    applyCurveSmoothing: boolean;
+    groupBy: GroupByOption;
+    viewableProperties: string[];
+}
diff --git a/Engine/Source/Programs/Horde/HordeDashboard/plugins/performancetrends/index.tsx b/Engine/Source/Programs/Horde/HordeDashboard/plugins/performancetrends/index.tsx
new file mode 100644
--- /dev/null
+++ b/Engine/Source/Programs/Horde/HordeDashboard/plugins/performancetrends/index.tsx
@@ -0,0 +1,19 @@
+// Copyright Epic Games, Inc. All Rights Reserved.
+
+import { MountType, registerHordePlugin } from "hordePlugins";
+import { PerformanceTrendsView } from "./PerformanceTrendView";
+
+registerHordePlugin({
+    id: "performancetrends",
+    routes: [
+        { path: "performancetrends", element: <PerformanceTrendsView/> },
+    ],
+    mounts: [
+        {
+            type: MountType.TopNav,
+            context: "Tools",
+            text: "Analytics - Performance Trends",
+            route: "/performancetrends"
+        }
+    ]
+});
\ No newline at end of file
diff --git a/Engine/Source/Programs/Horde/HordeDashboard/plugins/performancetrends/metrictypes/KeyStatsData.ts b/Engine/Source/Programs/Horde/HordeDashboard/plugins/performancetrends/metrictypes/KeyStatsData.ts
new file mode 100644
--- /dev/null
+++ b/Engine/Source/Programs/Horde/HordeDashboard/plugins/performancetrends/metrictypes/KeyStatsData.ts
@@ -0,0 +1,67 @@
+// Copyright Epic Games, Inc. All Rights Reserved.
+
+import { KeyStatsViewGenerator } from "../viewgenerators/KeyStatsViewGenerator";
+import { ISummaryMetric, MetricTypeRegistry, PerformanceTrendContext } from "./PerformanceTrendsTypes";
+import { viewable } from "./StatMetadata";
+
+/**
+ * A key stats metric data structure.
+ */
+export class KeyStatsData extends PerformanceTrendContext implements ISummaryMetric {
+
+    // #region -- Public Properties --
+
+    // Identification
+    csvId: string;
+
+    // Performance metrics
+    @viewable({ label: 'GPU Time', category: 'Performance', unit: 'ms', defaultVisible: true, precision: 2 })
+    gpuTimeAvg: number;
+
+    @viewable({ label: 'Game Thread', category: 'Timing', unit: 'ms', defaultVisible: true, precision: 2 })
+    gameThreadtimeAvg: number;
+
+    @viewable({ label: 'Render Thread', category: 'Timing', unit: 'ms', defaultVisible: true, precision: 2 })
+    renderThreadtimeAvg: number;
+
+    @viewable({ label: 'Peak Physical Memory', category: 'Memory', unit: 'MB', defaultVisible: true, precision: 1 })
+    physicalUsedMbMax: number;
+
+    @viewable({ label: 'Hitches/Min', category: 'Performance', defaultVisible: true, precision: 1 })
+    hitchesMin: number;
+
+    @viewable({ label: 'Hitch Time', category: 'Performance', unit: '%', precision: 2 })
+    hitchTimePercent: number;
+
+    @viewable({ label: 'MVP', category: 'Performance', precision: 2 })
+    mvp: number;
+
+    @viewable({ label: 'Dynamic Resolution Avg', category: 'Resolution', unit: '%', precision: 1 })
+    dynamicResolutionPercentageAvg: number;
+
+    @viewable({ label: 'Dynamic Resolution Max', category: 'Resolution', unit: '%', precision: 1 })
+    dynamicResolutionPercentageMax: number;
+
+    @viewable({ label: 'VSM Nanite Resolution', category: 'Resolution', precision: 2 })
+    dynamicResolutionVsmNanite: number;
+
+    // #endregion -- Public Properties --
+
+    // #region -- Implicit ISummaryMetricProvider interface --
+
+    /**
+     * @inheritdoc 
+     */
+    static readonly metricType = "KeyStats";
+
+    /**
+     * @inheritdoc 
+     */
+    constructor() {
+        super();
+    }
+
+    // #endregion -- Implicit ISummaryMetricProvider interface --
+};
+
+MetricTypeRegistry.register(KeyStatsData, new KeyStatsViewGenerator());
\ No newline at end of file
diff --git a/Engine/Source/Programs/Horde/HordeDashboard/plugins/performancetrends/metrictypes/PerformanceTrendsTypes.ts b/Engine/Source/Programs/Horde/HordeDashboard/plugins/performancetrends/metrictypes/PerformanceTrendsTypes.ts
new file mode 100644
--- /dev/null
+++ b/Engine/Source/Programs/Horde/HordeDashboard/plugins/performancetrends/metrictypes/PerformanceTrendsTypes.ts
@@ -0,0 +1,291 @@
+// Copyright Epic Games, Inc. All Rights Reserved.
+
+import { PerformanceTrendDataHandler } from "../PerformanceTrendsDataHandler";
+import { IMetricViewGenerator } from "../viewgenerators/PerformanceTrendRenderTypes";
+
+// #region -- Response Factory Types --
+
+/**
+ * Constant used to describe a placeholder non-resolution type.
+ */
+export const NO_RESOLVE_TYPE_FOUND: string = "no-resolved-type";
+
+/**
+ * Unified metric constraint which defines a summary metric that is validateable, and initializable.
+ */
+export type MetricConstraint = ISummaryMetric & IValidatable & IInitializable;
+
+// #region -- Metric Registry, Providers, and Factories --
+
+/**
+ * Facade pattern to simplify interaction. Use this to register a metrics type to participate in the Performance Trend Dashboard.
+ */
+export class MetricTypeRegistry {
+    /**
+     * Registers a metric type to participate in the metric retrieval, constrction, and visualization system of Performance Trends.
+     * @param metricType The metric type to register.
+     * @param viewGenerator The view generator instance to use in rendering metrics of @see metricType .
+     */
+    static register<T extends MetricConstraint>(metricType: ISummaryMetricProvider<T>, viewGenerator: IMetricViewGenerator<T>): void {
+        MetricFactory.register(metricType);
+        ViewGeneratorProvider.register(metricType, viewGenerator);
+        PerformanceTrendHandlerProvider.register(metricType);
+    }
+
+    /**
+     * Gets the data provider and view generator associated with the provided type, if it exists.
+     * @param metricType The metric type to request a data provider and / or view generator for.
+     * @returns The data provider and view generator for the requested type.
+     */
+    static get(metricType: string): { dataProvider: PerformanceTrendDataHandler<MetricConstraint> | undefined, viewGenerator: IMetricViewGenerator<MetricConstraint> | undefined } {
+        return {
+            dataProvider: PerformanceTrendHandlerProvider.get(metricType),
+            viewGenerator: ViewGeneratorProvider.get(metricType)
+        };
+    }
+
+    /**
+     * Creates an instance of the requested type, given a raw payload.
+     * @param metricType The name of the type to create.
+     * @param raw The payload to use to initialize the object.
+     * @returns A initialized instance of the metric, using the payload as the initialization context.
+     */
+    static create<T extends MetricConstraint>(metricType: string, raw: any): T | null {
+        return MetricFactory.create<T>(metricType, raw);
+    }
+}
+
+/**
+ * Perofrmance Trend Handler Provider which is responsible for registering and obtaining data handlers based on metric type.
+ */
+class PerformanceTrendHandlerProvider {
+    private static registry = new Map<string, PerformanceTrendDataHandler<MetricConstraint>>();
+
+    /**
+     * Registers a metric of a provided name, with a corresponding provider interface.
+     * @param metric The metric type to associated with a provider.
+     */
+    static register<T extends MetricConstraint>(metric: ISummaryMetricProvider<T>): void {
+        if (!this.registry.has(metric.metricType)) {
+            this.registry.set(metric.metricType, new PerformanceTrendDataHandler<T>(metric));
+        }
+    }
+
+    /**
+     * Gets the data handler for the given metric type.
+     * @param metricType The metric type key.
+     * @returns The handler, or undefined if not registered.
+     */
+    static get(metricType: string): PerformanceTrendDataHandler<MetricConstraint> | undefined {
+        return this.registry.get(metricType);
+    }
+}
+
+/**
+ * View Generator Provider which is responsible for registering and obtaining view generators based on metric type.
+ */
+class ViewGeneratorProvider {
+    private static registry = new Map<string, IMetricViewGenerator<MetricConstraint>>();
+
+    /**
+     * Registers a metric of a provided name, with a corresponding generator interface.
+     * @param metric The metric type to associate with a the generator.
+     * @param generator The generator to associate with the metric type.
+     */
+    static register<T extends MetricConstraint>(metric: ISummaryMetricProvider<T>, generator: IMetricViewGenerator<T>
+    ) {
+        this.registry.set(metric.metricType, generator);
+    }
+
+    /**
+     * Gets the view generator associated with a metric type.
+     * @param metricType The name of the metric.
+     * @returns An instance of the view generator if it exists, otherwise undefined.
+     */
+    static get(metricType: string): IMetricViewGenerator<MetricConstraint> | undefined {
+        return this.registry.get(metricType);
+    }
+}
+
+/**
+ * Metric factory used for creating instances of requested types.
+ */
+class MetricFactory {
+    private static registry = new Map<string, ISummaryMetricProvider<any>>();
+
+    /**
+     * Registers a metric of a provided name, with a corresponding construction interface.
+     * @param ctor The ctor to use when requested for a type.
+     */
+    static register<T extends MetricConstraint>(ctor: ISummaryMetricProvider<T>) {
+        this.registry.set(ctor.metricType, ctor);
+    }
+
+    /**
+     * Creates an object of a specific type based off of the name.
+     * @param typeName the type to create.
+     * @param raw The raw representation of the data.
+     * @returns The constructed object of the specific type @see T.
+     */
+    static create<T extends MetricConstraint>(typeName: string, raw: any): T | null {
+        const ctor = this.registry.get(typeName);
+
+        if (!ctor) {
+            throw new Error(`No metric registered for type '${typeName}'`);
+        }
+
+        const instance: T = new ctor();
+
+        Object.assign(instance, raw);
+        instance.initialize(raw);
+
+        if (!instance.validate()) {
+            console.warn("Instance failed validation - discarding.");
+
+            return null;
+        }
+
+        return instance;
+    }
+}
+
+// #endregion -- Metric Registry, Providers, and Factories --
+
+// #endregion -- Response Factory Types --
+
+// #region -- Summary Metric Interfaces --
+
+/**
+ * An interface used to represent a summary metric.
+ */
+export interface ISummaryMetric {
+    summaryName: string;
+}
+
+/**
+ * An interface used to represent the validity of a type.
+ */
+export interface IValidatable {
+    /**
+     * Validates the object.
+     * @returns True if this is a valid object, false otherwise.
+     */
+    validate(): boolean;
+}
+
+/**
+ * An interface used to represent an initializable type.
+ * This is necessary due to Object.assign as the base init for some types, and thus acts as a late init for non-property mapped fields.
+ */
+export interface IInitializable {
+    /**
+     * Late initialization function to handle more nuanced initialization.
+     * @param raw Raw input used for specialized non-property-mapped initialization.
+     */
+    initialize(raw: any): void;
+}
+
+/**
+ * An interface used to request and retrieve metrics.
+ */
+export interface ISummaryMetricProvider<T extends ISummaryMetric> {
+
+    /**
+     * Constructs a summary metric from a raw payload.
+     */
+    new(): T;
+
+    /**
+     * The metric type to request, and reconstruct.
+     */
+    readonly metricType: string;
+}
+
+// #endregion -- Summary Metric Interfaces --
+
+// #region -- Shared Metric Abstract Classes  --
+
+/**
+ * Base context for all telemetry.
+ */
+export abstract class BaseContext implements IInitializable, IValidatable {
+    eventName: string;
+    schemaVersion: number;
+
+    constructor() {
+    }
+
+    /**
+     * @inheritdoc 
+     */
+    validate(): boolean {
+        return true;
+    }
+
+    /**
+     * @inheritdoc 
+     */
+    initialize(raw: any) { }
+}
+
+/**
+ * Common context for correlatable metrics & telemetry.
+ */
+export abstract class CorrelatableContext extends BaseContext {
+    sessionId?: string;
+    sessionLabel?: string;
+
+    constructor() {
+        super();
+    }
+}
+
+/**
+ * Common context for Horde telemetry.
+ */
+export abstract class HordeContext extends CorrelatableContext {
+    hordeUrlStr: string;
+    streamId: string;
+    jobId?: string;
+    stepId?: string;
+    commitId: string;
+    commitIdOrdered: number;
+    isBuildMachine: boolean;
+
+    constructor() {
+        super();
+    }
+}
+
+/**
+ * Common context for performance trend metrics & telemetry.
+ */
+export abstract class PerformanceTrendContext extends HordeContext implements ISummaryMetric {
+    testId: string;
+    testName: string;
+    testConfigName: string;
+    gauntletTestType: string;
+    gauntletSubTest: string;
+    summaryName: string;
+    platform: string;
+    collated: boolean;
+    engineVersion: string;
+    engineReleaseVersion: string;
+    buildVersion: string;
+    computedStream: string;
+    startTimestamp: number;
+    endTimestamp: number;
+
+    constructor() {
+        super()
+    }
+
+    /**
+     * @inheritdoc 
+     */
+    override validate(): boolean {
+        return this.summaryName !== undefined;
+    }
+}
+
+// #endregion -- Shared Metric Abstract Classes  --
\ No newline at end of file
diff --git a/Engine/Source/Programs/Horde/HordeDashboard/plugins/performancetrends/metrictypes/StatMetadata.ts b/Engine/Source/Programs/Horde/HordeDashboard/plugins/performancetrends/metrictypes/StatMetadata.ts
new file mode 100644
--- /dev/null
+++ b/Engine/Source/Programs/Horde/HordeDashboard/plugins/performancetrends/metrictypes/StatMetadata.ts
@@ -0,0 +1,70 @@
+// Copyright Epic Games, Inc. All Rights Reserved.
+
+export interface PropertyMeta {
+    label: string;
+    category?: string;
+    defaultVisible?: boolean;
+    unit?: string;
+    precision?: number;
+}
+
+// Store metadata on the class constructor
+const viewableRegistry = new WeakMap<Function, Map<string, PropertyMeta>>();
+
+/**
+ * Decorator to mark a property as viewable in the UI.
+ * Works with legacy experimentalDecorators.
+ */
+export function viewable(meta: PropertyMeta): PropertyDecorator {
+    return function (target: Object, propertyKey: string | symbol): void {
+        const ctor = target.constructor;
+
+        if (!viewableRegistry.has(ctor)) {
+            viewableRegistry.set(ctor, new Map());
+        }
+
+        viewableRegistry.get(ctor)!.set(String(propertyKey), meta);
+    };
+}
+
+/**
+ * Get all viewable properties from a class.
+ */
+export function getViewableProperties(target: Function | object): Map<string, PropertyMeta> {
+    const ctor = typeof target === 'function' ? target : target.constructor;
+    return viewableRegistry.get(ctor) ?? new Map();
+}
+
+/**
+ * Get viewable properties grouped by category.
+ */
+export function getViewablePropertiesByCategory(target: Function | object): Map<string, [string, PropertyMeta][]> {
+    const properties = getViewableProperties(target);
+    const grouped = new Map<string, [string, PropertyMeta][]>();
+
+    for (const [fieldName, meta] of properties) {
+        const category = meta.category ?? 'Other';
+        if (!grouped.has(category)) {
+            grouped.set(category, []);
+        }
+        grouped.get(category)!.push([fieldName, meta]);
+    }
+
+    return grouped;
+}
+
+/**
+ * Get default visible property names.
+ */
+export function getDefaultVisibleProperties(target: Function | object): string[] {
+    const properties = getViewableProperties(target);
+    const result: string[] = [];
+
+    for (const [fieldName, meta] of properties) {
+        if (meta.defaultVisible === true) {
+            result.push(fieldName);
+        }
+    }
+
+    return result;
+}
\ No newline at end of file
diff --git a/Engine/Source/Programs/Horde/HordeDashboard/plugins/performancetrends/responses/FilterKeys.ts b/Engine/Source/Programs/Horde/HordeDashboard/plugins/performancetrends/responses/FilterKeys.ts
new file mode 100644
--- /dev/null
+++ b/Engine/Source/Programs/Horde/HordeDashboard/plugins/performancetrends/responses/FilterKeys.ts
@@ -0,0 +1,215 @@
+// Copyright Epic Games, Inc. All Rights Reserved.
+
+import { GetTestProjectResponse } from "./GetTestProjectResponse";
+
+/**
+ * Separates aggregate parts of key.
+ */
+export const KEY_SEPARATOR: string = "::";
+
+/**
+ * Decoded test identity key.
+ */
+export type DecodedTestIdentityKey =
+    {
+        testProject: string,
+        testIdentity: string
+    }
+
+/**
+ * Decoded stream key.
+ */
+export type DecodedStreamKey = {
+    testProject: string,
+    testIdentity: string,
+    summaryType: string,
+    stream: string
+}
+
+/**
+ * Decoded platform key.
+ */
+export type DecodedPlatformKey = {
+    testProject: string,
+    testIdentity: string,
+    summaryType: string,
+    platform: string
+}
+
+/**
+ * Decoded metric summary key.
+ */
+export type DecodedMetricSummaryKey = {
+    testProject: string,
+    testIdentity: string,
+    summaryType: string
+}
+
+/**
+ * Decoded test type key.
+ */
+export type DecodedTestTypeKey = {
+    testProject: string,
+    testIdentity: string,
+    metricSummaryType: string,
+    testType: string
+}
+
+/**
+ * Decoded viewable property key.
+ */
+export type DecodedViewablePropertyKey = {
+    testType: string,
+    propertyKey: string
+}
+
+/**
+ * Encodes a stream unique key based on the input @see GetTestProjectResponse .
+ * @param testProject The input test project to base the stream key off of.
+ * @returns The encoded stream key.
+ */
+export function encodeStreamKey(testProject: GetTestProjectResponse): string {
+    return `${testProject.testName}${KEY_SEPARATOR}${testProject.testIdentity}${KEY_SEPARATOR}${testProject.summaryType}${KEY_SEPARATOR}${testProject.computedStream}`;
+}
+
+/**
+ * Obtains the stream, metric summary type, test identity, and owning test project from an encoded input stream type key. 
+ * @param encodedStreamKey The encoded platform key.
+ * @returns The resulting test project, test identity, metric summary type, and platform of the encoded key.
+ */
+export function decodeStreamKey(encodedStreamKey: string): DecodedStreamKey | null {
+    if (!encodedStreamKey || !encodedStreamKey.includes(KEY_SEPARATOR)) {
+        console.warn(`Invalid stream key: ${encodedStreamKey}`);
+
+        return null;
+    }
+
+    const [testProject, testIdentity, summaryType, stream] = encodedStreamKey.split(KEY_SEPARATOR);
+    return { testProject, testIdentity, summaryType, stream };
+}
+
+/**
+ * Encodes a platform unique key based on the input @see GetTestProjectPlatformResponse .
+ * @param testProject The input test project to base the platform key off of.
+ * @returns The encoded platform key.
+ */
+export function encodePlatformKey(testProject: GetTestProjectResponse, platform: string): string {
+    return `${testProject.testName}${KEY_SEPARATOR}${testProject.testIdentity}${KEY_SEPARATOR}${testProject.summaryType}${KEY_SEPARATOR}${platform}`;
+}
+
+/**
+ * Obtains the platform, metric summary type, test identity, and owning test project from an encoded input platform type key.
+ * @param encodedPlatformKey The encoded platform key.
+ * @returns The resulting test project, test identity, metric summary type, and platform of the encoded key.
+ */
+export function decodePlatformKey(encodedPlatformKey: string): DecodedPlatformKey | null {
+    if (!encodedPlatformKey || !encodedPlatformKey.includes(KEY_SEPARATOR)) {
+        console.warn(`Invalid platform key: ${encodedPlatformKey}`);
+
+        return null;
+    }
+
+    const [testProject, testIdentity, summaryType, platform] = encodedPlatformKey.split(KEY_SEPARATOR);
+    return { testProject, testIdentity, summaryType, platform };
+}
+
+/**
+ * Encodes a test type unique key based on the input @see GetTestProjectResponse .
+ * @param testProject The input test project to base the test type key off of.
+ * @returns The encoded test type key.
+ */
+export function encodeTestTypeKey(testProject: GetTestProjectResponse): string {
+    return `${testProject.testName}${KEY_SEPARATOR}${testProject.testIdentity}${KEY_SEPARATOR}${testProject.summaryType}${KEY_SEPARATOR}${testProject.testType}`;
+}
+
+/**
+ * Obtains the test type, test identity, and owning test project from an encoded input test type key.
+ * @param encodedTestType The encoded test type key.
+ * @returns The resulting test project, test identity, and test type of the encoded key.
+ */
+export function decodeTestTypeKey(encodedTestType: string): DecodedTestTypeKey | null {
+    if (!encodedTestType || !encodedTestType.includes(KEY_SEPARATOR)) {
+        console.warn(`Invalid test type key: ${encodedTestType}`);
+
+        return null;
+    }
+
+    const [testProject, testIdentity, metricSummaryType, testType] = encodedTestType.split(KEY_SEPARATOR);
+    return { testProject, testIdentity, metricSummaryType, testType };
+}
+
+/**
+ * Encodes a metric summary type unique key based on the input @see GetTestProjectResponse .
+ * @param testProject The input test project to base the metric summary type key off of.
+ * @returns The encoded metric summary key.
+ */
+export function encodeMetricTypeKey(testProject: GetTestProjectResponse): string {
+    return `${testProject.testName}${KEY_SEPARATOR}${testProject.testIdentity}${KEY_SEPARATOR}${testProject.summaryType}`;
+}
+
+/**
+ * Obtains the metric summary type, test identity, and owning test project from an encoded input metric summary type key.
+ * @param encodedMetricTypeKey The encoded metric summary type key.
+ * @returns The resulting test project, test identity, and metric summary type of the encoded key.
+ */
+export function decodeMetricTypeKey(encodedMetricTypeKey: string): DecodedMetricSummaryKey | null {
+    if (!encodedMetricTypeKey || !encodedMetricTypeKey.includes(KEY_SEPARATOR)) {
+        console.warn(`Invalid metric type key: ${encodedMetricTypeKey}`);
+
+        return null;
+    }
+
+    const [testProject, testIdentity, summaryType] = encodedMetricTypeKey.split(KEY_SEPARATOR);
+    return { testProject, testIdentity, summaryType };
+}
+
+/**
+ * Encodes a test identity unique key based on the input @see GetTestProjectResponse .
+ * @param testProject The input test project to base the test identity key off of.
+ * @returns The encoded test identity key.
+ */
+export function encodeTestIdentity(testProject: GetTestProjectResponse): string {
+    return `${testProject.testName}${KEY_SEPARATOR}${testProject.testIdentity}`;
+}
+
+/**
+ * Obtains the test identity, and owning test project from an encoded input test identity key.
+ * @param encodedTestIdentity The encoded test identity key.
+ * @returns The resulting test project and test identity of the encoded key.
+ */
+export function decodeTestIdentity(encodedTestIdentity: string): DecodedTestIdentityKey | null {
+    if (!encodedTestIdentity.includes(KEY_SEPARATOR)) {
+        console.warn(`Invalid test identity key: ${encodedTestIdentity}`);
+
+        return null;
+    }
+
+    const [testProject, testIdentity] = encodedTestIdentity.split(KEY_SEPARATOR);
+    return { testProject, testIdentity };
+}
+
+/**
+ * Encodes a viewable property unique key based on test type and property key.
+ * @param testType The test type (e.g., "PerfTest").
+ * @param propertyKey The property key (e.g., "gpuTimeAvg").
+ * @returns The encoded viewable property key.
+ */
+export function encodeViewablePropertyKey(testType: string, propertyKey: string): string {
+    return `${testType}${KEY_SEPARATOR}${propertyKey}`;
+}
+
+/**
+ * Obtains the test type and property key from an encoded viewable property key.
+ * @param encodedViewablePropertyKey The encoded viewable property key.
+ * @returns The resulting test type and property key, or null if invalid.
+ */
+export function decodeViewablePropertyKey(encodedViewablePropertyKey: string): DecodedViewablePropertyKey | null {
+    if (!encodedViewablePropertyKey || !encodedViewablePropertyKey.includes(KEY_SEPARATOR)) {
+        console.warn(`Invalid viewable property key: ${encodedViewablePropertyKey}`);
+
+        return null;
+    }
+
+    const [testType, propertyKey] = encodedViewablePropertyKey.split(KEY_SEPARATOR);
+    return { testType, propertyKey };
+}
\ No newline at end of file
diff --git a/Engine/Source/Programs/Horde/HordeDashboard/plugins/performancetrends/responses/GetTestProjectResponse.ts b/Engine/Source/Programs/Horde/HordeDashboard/plugins/performancetrends/responses/GetTestProjectResponse.ts
new file mode 100644
--- /dev/null
+++ b/Engine/Source/Programs/Horde/HordeDashboard/plugins/performancetrends/responses/GetTestProjectResponse.ts
@@ -0,0 +1,28 @@
+// Copyright Epic Games, Inc. All Rights Reserved.
+
+/** 
+ * Response describing a Performance Report Tool Test Project.
+ * */
+export type GetTestProjectResponse = {
+    summaryType: string;
+    testName: string;
+    testIdentity: string;
+    testType: string;
+    computedStream: string;
+}
+
+/** 
+ * Response describing a distinct platform entry for a Performance Report Tool Test Project entry.
+ * */
+export type GetTestProjectPlatformResponse = {
+    owningTestProject: GetTestProjectResponse;
+    platforms: string[];
+}
+
+/** 
+ * Response describing a distinct commit entry for a Performance Report Tool Test Project's platform entry.
+ * */
+export type GetTestProjectCommitResponse = {
+    owningTestProject: GetTestProjectResponse;
+    commitIds: number[];
+}
\ No newline at end of file
diff --git a/Engine/Source/Programs/Horde/HordeDashboard/plugins/performancetrends/viewgenerators/KeyStatsViewGenerator.tsx b/Engine/Source/Programs/Horde/HordeDashboard/plugins/performancetrends/viewgenerators/KeyStatsViewGenerator.tsx
new file mode 100644
--- /dev/null
+++ b/Engine/Source/Programs/Horde/HordeDashboard/plugins/performancetrends/viewgenerators/KeyStatsViewGenerator.tsx
@@ -0,0 +1,485 @@
+// Copyright Epic Games, Inc. All Rights Reserved.
+
+import { graphColors } from "hordePlugins/analytics/telemetryData";
+import { LineGraph, MetricSelectionHandler, SeriesConfig } from "../components/LineGraphComponent";
+import { KeyStatsData } from "../metrictypes/KeyStatsData";
+import { IMetricDetailedView, IMetricView, IMetricViewGenerator } from "./PerformanceTrendRenderTypes";
+import { GroupByOption, PerformanceTrendUIOptionsState } from "../filters/PerformanceTrendUIOptionsState";
+import { Label } from "@fluentui/react";
+import { KEY_SEPARATOR } from "../responses/FilterKeys";
+import { getDefaultVisibleProperties, getViewableProperties, PropertyMeta } from "../metrictypes/StatMetadata";
+import { observable } from "mobx";
+import { ColumnDef, CSVTable } from "hordePlugins/structuredanalytics/CSVComponent";
+
+// #region -- Helpers --
+
+/**
+ * Type-safe property name accessor - ensures property exists on type at compile time.
+ */
+function keyOf<T>(key: keyof T): keyof T {
+    return key;
+}
+
+/**
+ * Helper to build a map of propertyKey -> Set of test types where it's enabled.
+ */
+function buildPropertyToTestTypesMap(enabledProps: Record<string, string>): Map<string, Set<string>> {
+    const propertyToTestTypes = new Map<string, Set<string>>();
+    for (const compositeKey of Object.keys(enabledProps)) {
+        const [testType, propertyKey] = compositeKey.split('::');
+        if (!propertyToTestTypes.has(propertyKey)) {
+            propertyToTestTypes.set(propertyKey, new Set());
+        }
+        propertyToTestTypes.get(propertyKey)!.add(testType);
+    }
+    return propertyToTestTypes;
+}
+
+/**
+ * Build y-axis label from PropertyMeta.
+ */
+function buildYAxisLabel(meta: PropertyMeta): string {
+    if (meta.unit) {
+        return `${meta.label} (${meta.unit})`;
+    }
+    return meta.label;
+}
+
+/**
+ * Dynamic property accessor - gets value from object using property key.
+ */
+function getPropertyValue(obj: KeyStatsData, propertyKey: string): number {
+    return (obj as unknown as Record<string, number>)[propertyKey];
+}
+
+// #endregion -- Helpers --
+
+/**
+ * KeyStats view generator.
+ * Uses @viewable decorator metadata to dynamically generate views for each enabled property.
+ */
+export class KeyStatsViewGenerator implements IMetricViewGenerator<KeyStatsData> {
+
+    // #region -- Private Helpers --
+
+    private getGroupKey(metric: KeyStatsData, groupBy: GroupByOption): string {
+        const platform = metric.platform ?? 'Unknown';
+        const stream = metric.computedStream ?? 'Unknown';
+        const testType = metric.gauntletSubTest ?? 'Unknown';
+
+        switch (groupBy) {
+            case 'platform':
+                return platform;
+            case 'stream':
+                return stream;
+            case 'testType':
+                return testType;
+            case 'platform+stream':
+                return `${platform}${KEY_SEPARATOR}${stream}`;
+            case 'platform+testType':
+                return `${platform}${KEY_SEPARATOR}${testType}`;
+            case 'stream+testType':
+                return `${stream}${KEY_SEPARATOR}${testType}`;
+            case 'platform+stream+testType':
+                return `${platform}${KEY_SEPARATOR}${stream}${KEY_SEPARATOR}${testType}`;
+        }
+    }
+
+    private getGroupLabel(key: string, groupBy: GroupByOption): string {
+        const parts = key.split(KEY_SEPARATOR);
+        switch (groupBy) {
+            case 'platform':
+            case 'stream':
+            case 'testType':
+                return key;
+            case 'platform+stream':
+            case 'platform+testType':
+            case 'stream+testType':
+                return `${parts[0]}, ${parts[1]}`;
+            case 'platform+stream+testType':
+                return `${parts[0]}, ${parts[1]}, ${parts[2]}`;
+        }
+    }
+
+    private groupMetrics(metricsToGroup: KeyStatsData[], groupBy: GroupByOption): Map<string, KeyStatsData[]> {
+        const grouped = new Map<string, KeyStatsData[]>();
+
+        for (const metric of metricsToGroup) {
+            const key = this.getGroupKey(metric, groupBy);
+
+            if (!grouped.has(key)) {
+                grouped.set(key, []);
+            }
+            grouped.get(key)!.push(metric);
+        }
+
+        return grouped;
+    }
+
+    // #endregion -- Private Helpers --
+
+    /**
+     * @inheritdoc
+     */
+    getViews(metrics: KeyStatsData[], uiState: PerformanceTrendUIOptionsState, selectionHandler?: MetricSelectionHandler<KeyStatsData>): IMetricView[] {
+
+        const groupBy: GroupByOption = uiState.groupByOption;
+        const metricsByGroup = this.groupMetrics(metrics, groupBy);
+        const sortedGroupKeys = Array.from(metricsByGroup.keys()).sort();
+
+        // Get viewable properties metadata from @viewable decorators
+        const viewableProps = getViewableProperties(KeyStatsData);
+
+        // Check which properties are enabled for which test types
+        const propertyToTestTypes = buildPropertyToTestTypesMap(uiState.enabledViewableProperties);
+
+        const createGroupedSeries = <T extends KeyStatsData>(propertyKey: string, yAccessor: (m: T) => number, labelPrefix: string, baseColorIndex: number): SeriesConfig<T>[] => {
+            const enabledTestTypes = propertyToTestTypes.get(propertyKey);
+            if (!enabledTestTypes) return [];
+
+            return sortedGroupKeys.map((groupKey, idx) => {
+                const groupData = metricsByGroup.get(groupKey)! as T[];
+                const filteredData = groupData.filter(m => enabledTestTypes.has(m.gauntletSubTest));
+                return {
+                    data: filteredData,
+                    yAccessor,
+                    label: `${labelPrefix}\n(${this.getGroupLabel(groupKey, groupBy)})`,
+                    color: graphColors[(baseColorIndex + idx) % graphColors.length]
+                };
+            }).filter(s => s.data.length > 0);
+        };
+
+        const changelistAccessor = (value: KeyStatsData, _index: number): number => value.commitIdOrdered;
+
+        // Get metadata for each property (type-safe property names)
+        const gpuMeta = viewableProps.get(keyOf<KeyStatsData>('gpuTimeAvg'))!;
+        const hitchesMeta = viewableProps.get(keyOf<KeyStatsData>('hitchesMin'))!;
+        const memoryMeta = viewableProps.get(keyOf<KeyStatsData>('physicalUsedMbMax'))!;
+        const gameThreadMeta = viewableProps.get(keyOf<KeyStatsData>('gameThreadtimeAvg'))!;
+        const renderThreadMeta = viewableProps.get(keyOf<KeyStatsData>('renderThreadtimeAvg'))!;
+
+        // Create series arrays for each metric type (filtered by enabled properties)
+        const gpuSeries = createGroupedSeries<KeyStatsData>(
+            keyOf<KeyStatsData>('gpuTimeAvg'),
+            (m) => m.gpuTimeAvg,
+            gpuMeta.label,
+            10
+        );
+
+        const hitchesPerMinSeries = createGroupedSeries<KeyStatsData>(
+            keyOf<KeyStatsData>('hitchesMin'),
+            (m) => m.hitchesMin,
+            hitchesMeta.label,
+            11
+        );
+
+        const memoryUsedMaxSeries = createGroupedSeries<KeyStatsData>(
+            keyOf<KeyStatsData>('physicalUsedMbMax'),
+            (m) => m.physicalUsedMbMax,
+            memoryMeta.label,
+            13
+        );
+
+        const gameThreadSeries = createGroupedSeries<KeyStatsData>(
+            keyOf<KeyStatsData>('gameThreadtimeAvg'),
+            (m) => m.gameThreadtimeAvg,
+            gameThreadMeta.label,
+            0
+        );
+
+        const renderThreadSeries = createGroupedSeries<KeyStatsData>(
+            keyOf<KeyStatsData>('renderThreadtimeAvg'),
+            (m) => m.renderThreadtimeAvg,
+            renderThreadMeta.label,
+            5
+        );
+
+        // Calculate averages across all platforms
+        const gpuTimes: number[] = metrics.map(m => m.gpuTimeAvg);
+        const memoryUsedMax: number[] = metrics.map(m => m.physicalUsedMbMax);
+        const hitchesPerMin: number[] = metrics.map(m => m.hitchesMin);
+        const avgGpuTimes = gpuTimes.length === 0 ? undefined : gpuTimes.reduce((s, v) => s + v, 0) / gpuTimes.length;
+        const avgHitchesPerMin = hitchesPerMin.length === 0 ? undefined : hitchesPerMin.reduce((s, v) => s + v, 0) / hitchesPerMin.length;
+        const avgMemoryUsed = memoryUsedMax.length === 0 ? undefined : memoryUsedMax.reduce((s, v) => s + v, 0) / memoryUsedMax.length;
+
+        const GpuTimeView = ({ width, height }: { width: number; height: number }) => {
+            return (
+                <div>
+                    <div>
+                        <span>{gpuMeta.label} (all metrics)</span><br />
+                        <span>(Avg) {avgGpuTimes?.toFixed(gpuMeta.precision ?? 1)} {gpuMeta.unit ? `(${gpuMeta.unit})` : ''}</span>
+                    </div>
+                    <LineGraph
+                        series={gpuSeries}
+                        xAxisLabel="Changelist"
+                        yAxisLabel={buildYAxisLabel(gpuMeta)}
+                        width={width}
+                        height={height}
+                        yAxisZeroScale={true}
+                        applyCurveSmoothing={uiState.curveGraphLine}
+                        onClick={selectionHandler}
+                        xAccessor={changelistAccessor}
+                    />
+                </div>
+            );
+        };
+
+        const HitchesPerMinView = ({ width, height }: { width: number; height: number }) => {
+            return (
+                <div>
+                    <div>
+                        <span>{hitchesMeta.label}</span> <br />
+                        <span>(Avg) {avgHitchesPerMin?.toFixed(hitchesMeta.precision ?? 1)}</span>
+                    </div>
+                    <LineGraph
+                        series={hitchesPerMinSeries}
+                        xAxisLabel="Changelist"
+                        yAxisLabel={hitchesMeta.label}
+                        width={width}
+                        height={height}
+                        yAxisZeroScale={true}
+                        applyCurveSmoothing={uiState.curveGraphLine}
+                        onClick={selectionHandler}
+                        xAccessor={changelistAccessor}
+                    />
+                </div>
+            );
+        };
+
+        const ThreadTimesView = ({ width, height }: { width: number; height: number }) => {
+            return (
+                <div>
+                    <div>
+                        <span>Thread Times</span> <br />
+                    </div>
+                    <LineGraph
+                        series={[...gameThreadSeries, ...renderThreadSeries]}
+                        xAxisLabel="Changelist"
+                        yAxisLabel={buildYAxisLabel(gameThreadMeta)}
+                        width={width}
+                        height={height}
+                        yAxisZeroScale={true}
+                        applyCurveSmoothing={uiState.curveGraphLine}
+                        onClick={selectionHandler}
+                        xAccessor={changelistAccessor}
+                    />
+                </div>
+            );
+        };
+
+        const PeakPhysicalMemoryView = ({ width, height }: { width: number; height: number }) => {
+            return (
+                <div>
+                    <div>
+                        <span>{memoryMeta.label}</span> <br />
+                        <span>(Avg) {avgMemoryUsed?.toFixed(memoryMeta.precision ?? 1)} {memoryMeta.unit ? `(${memoryMeta.unit})` : ''}</span>
+                    </div>
+                    <LineGraph
+                        series={memoryUsedMaxSeries}
+                        xAxisLabel="Changelist"
+                        yAxisLabel={buildYAxisLabel(memoryMeta)}
+                        width={width}
+                        height={height}
+                        yAxisZeroScale={true}
+                        applyCurveSmoothing={uiState.curveGraphLine}
+                        onClick={selectionHandler}
+                        xAccessor={changelistAccessor}
+                    />
+                </div>
+            );
+        };
+
+        const views: IMetricView[] = [];
+
+        if (gpuSeries.length > 0) {
+            views.push({ label: gpuMeta.label, Render: GpuTimeView });
+        }
+        if (hitchesPerMinSeries.length > 0) {
+            views.push({ label: hitchesMeta.label, Render: HitchesPerMinView });
+        }
+        if (memoryUsedMaxSeries.length > 0) {
+            views.push({ label: memoryMeta.label, Render: PeakPhysicalMemoryView });
+        }
+        if (gameThreadSeries.length > 0 || renderThreadSeries.length > 0) {
+            views.push({ label: "Thread Timings", Render: ThreadTimesView });
+        }
+
+        return views;
+    }
+
+    /**
+     * Dynamically generates views based on enabled properties using @viewable decorator metadata.
+     */
+    getViewsDynamic(metrics: KeyStatsData[], uiState: PerformanceTrendUIOptionsState, selectionHandler?: MetricSelectionHandler<KeyStatsData>): IMetricView[] {
+
+        const enabledProps = uiState.enabledViewableProperties;
+
+        // If no properties are enabled, show nothing
+        if (Object.keys(enabledProps).length === 0) {
+            return [];
+        }
+
+        const propertyToTestTypes = buildPropertyToTestTypesMap(enabledProps);
+        const groupBy: GroupByOption = uiState.groupByOption;
+        const changelistAccessor = (value: KeyStatsData, _index: number): number => value.commitIdOrdered;
+
+        // Get viewable properties metadata from @viewable decorators
+        const viewableProps = getViewableProperties(KeyStatsData);
+
+        // Create a view for each enabled property
+        const views: IMetricView[] = [];
+        let colorOffset = 0;
+
+        for (const [propertyKey, enabledTestTypes] of propertyToTestTypes) {
+            const meta = viewableProps.get(propertyKey);
+
+            if (!meta) {
+                continue;
+            }
+
+            // Filter metrics to only include those from test types where THIS property is enabled
+            const propertyMetrics = metrics.filter(m => enabledTestTypes.has(m.gauntletSubTest));
+
+            if (propertyMetrics.length === 0) {
+                continue;
+            }
+
+            // Group the filtered metrics for this property
+            const metricsByGroup = this.groupMetrics(propertyMetrics, groupBy);
+            const sortedGroupKeys = Array.from(metricsByGroup.keys()).sort();
+
+            // Use incrementing color offset for each property
+            const baseColorIndex = colorOffset;
+            colorOffset += sortedGroupKeys.length;
+
+            // Create series for this property using dynamic accessor
+            const series: SeriesConfig<KeyStatsData>[] = sortedGroupKeys.map((groupKey, idx) => ({
+                data: metricsByGroup.get(groupKey)!,
+                yAccessor: (m: KeyStatsData) => getPropertyValue(m, propertyKey),
+                label: `${meta.label}\n(${this.getGroupLabel(groupKey, groupBy)})`,
+                color: graphColors[(baseColorIndex + idx) % graphColors.length]
+            }));
+
+            // Calculate average for display
+            const values = propertyMetrics.map(m => getPropertyValue(m, propertyKey)).filter(v => v !== undefined && !isNaN(v));
+            const avg = values.length > 0 ? values.reduce((s, v) => s + v, 0) / values.length : undefined;
+
+            // Build y-axis label from metadata
+            const yAxisLabel = buildYAxisLabel(meta);
+
+            const PropertyView = ({ width, height }: { width: number; height: number }) => {
+                return (
+                    <div>
+                        <div>
+                            <span>{meta.label}</span><br />
+                            <span>(Avg) {avg?.toFixed(meta.precision ?? 1)} {meta.unit ? `(${meta.unit})` : ''}</span>
+                        </div>
+                        <LineGraph
+                            series={series}
+                            xAxisLabel="Changelist"
+                            yAxisLabel={yAxisLabel}
+                            width={width}
+                            height={height}
+                            yAxisZeroScale={true}
+                            applyCurveSmoothing={uiState.curveGraphLine}
+                            onClick={selectionHandler}
+                            xAccessor={changelistAccessor}
+                        />
+                    </div>
+                );
+            };
+
+            views.push({
+                label: meta.label,
+                Render: PropertyView,
+            });
+        }
+
+        return views;
+    }
+
+    /**
+     * @inheritdoc
+     */
+    getSelectedDetailView(metrics: KeyStatsData[], _uiState: PerformanceTrendUIOptionsState): IMetricDetailedView {
+        return {
+            Render: (props) => <KeyStatsDetailView metrics={metrics} {...props} />
+        };
+    }
+
+    /**
+     * @inheritdoc 
+     */
+    getViewableProperties(): Map<string, PropertyMeta> {
+        return getViewableProperties(KeyStatsData);
+    }
+
+    /**
+     * @inheritdoc 
+     */
+    getDefaultVisibleProperties(): string[] {
+        return getDefaultVisibleProperties(KeyStatsData);
+    }
+};
+
+/**
+ * Detail view component for KeyStats data.
+ * Extracted as standalone component to avoid MobX observer warnings.
+ */
+const KeyStatsDetailView: React.FC<{ metrics: KeyStatsData[], width: number, height: number }> = ({ metrics, width, height }) => {
+    if (!metrics || metrics.length === 0) {
+        return <span>No data selected</span>;
+    }
+
+    const ensureTrailingSlash = (url: string): string => url.endsWith('/') ? url : `${url}/`;
+
+    const columns: (string | ColumnDef<KeyStatsData>)[] = [
+        {
+            key: "hordeLink",
+            label: "Horde Step",
+            render: (_, row, cell) => {
+                const url = `${ensureTrailingSlash(row.hordeUrlStr)}job/${row.jobId}${row.stepId ? `?step=${row.stepId}` : ""}`;
+                cell.append("a")
+                    .attr("href", url)
+                    .attr("target", "_blank")
+                    .attr("rel", "noopener noreferrer")
+                    .text("View Step")
+                    .style("width", "140px")
+            }
+        },
+        {
+            key: "csvId",
+            label: "Csv Id",
+            render: (value, row, cell) => {
+                cell.style("font-weight", "bold")
+                    .text(value);
+            }
+        },
+        {
+            key: "gauntletSubTest",
+            label: "Test Type"
+        },
+        "platform",
+        "gpuTimeAvg",
+        "gameThreadtimeAvg",
+        "renderThreadtimeAvg",
+        "physicalUsedMbMax",
+        "hitchesMin",
+        "hitchTimePercent",
+        "mvp",
+        "dynamicResolutionPercentageAvg"
+    ];
+
+    // Wrap data in observable to satisfy CSVTable's observer wrapper
+    const dataProvider = observable({ data: metrics });
+
+    return (
+        <>
+            <Label>Key Stats</Label>
+            <CSVTable
+                dataProvider={dataProvider}
+                columnProvider={{ columns }}
+            />
+        </>
+    );
+};
\ No newline at end of file
diff --git a/Engine/Source/Programs/Horde/HordeDashboard/plugins/performancetrends/viewgenerators/PerformanceTrendRenderTypes.tsx b/Engine/Source/Programs/Horde/HordeDashboard/plugins/performancetrends/viewgenerators/PerformanceTrendRenderTypes.tsx
new file mode 100644
--- /dev/null
+++ b/Engine/Source/Programs/Horde/HordeDashboard/plugins/performancetrends/viewgenerators/PerformanceTrendRenderTypes.tsx
@@ -0,0 +1,58 @@
+// Copyright Epic Games, Inc. All Rights Reserved.
+
+import { MetricSelectionHandler } from "../components/LineGraphComponent";
+import { PerformanceTrendUIOptionsState } from "../filters/PerformanceTrendUIOptionsState";
+import { ISummaryMetric } from "../metrictypes/PerformanceTrendsTypes";
+import { PropertyMeta } from "../metrictypes/StatMetadata";
+
+/**
+ * Interface that defines a render action.
+ */
+export interface IMetricView {
+    /** Optional label for the metric */
+    label?: string;
+
+    /** 
+     * React component representing the visualization of this metric.
+     * @returns React component to render element.
+    */
+    Render: React.FC<{ width: number; height: number }>;
+}
+
+/**
+ * Interface that defines a detailed render action.
+ */
+export interface IMetricDetailedView {
+    Render: React.FC<{ width: number; height: number }>;
+}
+
+/**
+ * Interface that describes a class capable of generating a view representing @see TMetric.
+ */
+export interface IMetricViewGenerator<TMetric extends ISummaryMetric> {
+    /**
+     * Gets a list of views for a specific metric.
+     * @returns A list of @see IMetricView objects that can be rendered.
+     */
+    getViews(metrics: TMetric[], uiState: PerformanceTrendUIOptionsState, selectionHandler?: MetricSelectionHandler<TMetric>): IMetricView[];
+
+    /**
+     * Gets a view that represents a detailed perspective when a set of metrics is selected for analysis.
+     * @param metrics The metrics to provide a detailed view for.
+     * @param uiState The UI options.
+     * @returns A single @see IMetricDetailedView representing a detailed analysis.
+     */
+    getSelectedDetailView(metrics: TMetric[], uiState: PerformanceTrendUIOptionsState): IMetricDetailedView;
+
+    /**
+     * Gets the viewable properties for this metric type.
+     * @returns A map of property key to metadata.
+     */
+    getViewableProperties(): Map<string, PropertyMeta>;
+
+    /**
+     * Gets the default visible property keys.
+     * @returns Array of property keys visible by default.
+     */
+    getDefaultVisibleProperties(): string[];
+}
\ No newline at end of file
diff --git a/Engine/Source/Programs/Horde/HordeDashboard/plugins/registry.ts b/Engine/Source/Programs/Horde/HordeDashboard/plugins/registry.ts
--- a/Engine/Source/Programs/Horde/HordeDashboard/plugins/registry.ts
+++ b/Engine/Source/Programs/Horde/HordeDashboard/plugins/registry.ts
@@ -3,3 +3,5 @@
 import "./exampleplugin"
 import "./analytics"
 import "./test-automation"
+import "./performancetrends"
+import "./structuredanalytics"
diff --git a/Engine/Source/Programs/Horde/HordeDashboard/plugins/structuredanalytics/AnalyticsView.tsx b/Engine/Source/Programs/Horde/HordeDashboard/plugins/structuredanalytics/AnalyticsView.tsx
new file mode 100644
--- /dev/null
+++ b/Engine/Source/Programs/Horde/HordeDashboard/plugins/structuredanalytics/AnalyticsView.tsx
@@ -0,0 +1,120 @@
+// Copyright Epic Games, Inc. All Rights Reserved.
+
+import { Spinner, SpinnerSize, Stack } from "@fluentui/react"
+import { getJobSummaries, JobSummary } from "./api";
+import { makeObservable, observable } from "mobx";
+import { observer } from "mobx-react-lite";
+import React, { useEffect } from "react";
+import { getHordeStyling } from "horde/styles/Styles";
+import { CSVTable, IColumnProvider, IDataProvider } from "./CSVComponent";
+import { useWindowSize } from "horde/base/utilities/hooks";
+import { Breadcrumbs } from "horde/components/Breadcrumbs";
+import { TopNav } from "horde/components/TopNav";
+import { Text } from '@fluentui/react';
+
+// #region -- Data Handler --
+
+class AnalyticsHandler implements IDataProvider, IColumnProvider {
+
+    constructor() {
+        makeObservable(this);
+        this.columns = [];
+        this.data = [];
+    }
+
+    // #region -- IColumnProvider --
+
+    @observable
+    columns: string[] = [];
+
+    // #endregion -- IColumnProvider --
+
+    // #region -- IDataProvider --
+
+    @observable
+    data: JobSummary[] = [];
+
+    // #endregion -- IDataProvider --
+
+    async query() {
+        this.querying = true;
+
+        try {
+            this.jobSummaryData = await getJobSummaries();
+        } catch (errorReason) {
+            console.log(errorReason);
+        }
+
+        this.querying = false;
+        this.data = this.jobSummaryData;
+        this.columns = this.jobSummaryData.length >= 1 ? Object.keys(this.jobSummaryData[0]) : [];
+    }
+
+    // #region -- Local Members --
+
+    @observable
+    querying: boolean;
+
+    jobSummaryData: JobSummary[] = [];
+
+    // #endregion -- Local Members --
+}
+
+// #endregion --  Data Handler --
+
+// #region -- Visual Components --
+
+/**
+ * Analytics prototyping view.
+ * @returns React component.
+ */
+export const AnalyticsView: React.FC = observer(() => {
+    useEffect(() => {
+        handler.query();
+    }, []);
+
+    const { hordeClasses, modeColors } = getHordeStyling();
+    const windowSize = useWindowSize();
+
+    return (
+        <Stack className={hordeClasses.horde}>
+            <TopNav />
+            <Breadcrumbs items={[{ text: 'Analytics Sandbox' }]} />
+            <Stack horizontal>
+                <div key={`windowsize_automationview_${windowSize.width}_${windowSize.height}`} style={{ width: 0, flexShrink: 0, backgroundColor: modeColors.background }} />
+                <Stack horizontalAlign="center" grow styles={{ root: { width: "100%", padding: 12, backgroundColor: modeColors.background } }}>
+                    <Stack styles={{ root: { width: "100%" } }}>
+                        <Stack horizontal styles={{ root: { minHeight: '85vh' } }} >
+                            <Stack id="parent-analytics-trends" styles={{
+                                root: {
+                                    overflowX: "auto",
+                                    overflowY: "visible",
+                                    minWidth: 0,
+                                    position: "relative",
+                                    border: `${modeColors.content} solid 2px`
+                                }
+                            }}>
+                                {handler.querying ? (
+                                    <Spinner style={{ paddingTop: 16 }} size={SpinnerSize.large} />
+                                ) : handler.data.length === 0 ? (
+                                    <Text
+                                        variant="medium"
+                                        styles={{ root: { paddingTop: 16, color: modeColors.textSecondary } }}
+                                    >
+                                        No data available for the selected filters.
+                                    </Text>
+                                ) : (
+                                    <CSVTable dataProvider={handler} columnProvider={handler} />
+                                )}
+                            </Stack>
+                        </Stack>
+                    </Stack>
+                </Stack>
+            </Stack>
+        </Stack>
+    )
+});
+
+// #endregion -- Visual Components --
+
+let handler: AnalyticsHandler = new AnalyticsHandler();
\ No newline at end of file
diff --git a/Engine/Source/Programs/Horde/HordeDashboard/plugins/structuredanalytics/CSVComponent.tsx b/Engine/Source/Programs/Horde/HordeDashboard/plugins/structuredanalytics/CSVComponent.tsx
new file mode 100644
--- /dev/null
+++ b/Engine/Source/Programs/Horde/HordeDashboard/plugins/structuredanalytics/CSVComponent.tsx
@@ -0,0 +1,380 @@
+// Copyright Epic Games, Inc. All Rights Reserved.
+
+import { getColorRecords } from "horde/components/buildhealth/stepoutcome/StepOutcomeSharedUIComponents";
+import { getHordeStyling, preloadFonts } from "horde/styles/Styles";
+import { observer } from "mobx-react-lite";
+import { useCallback, useEffect, useMemo, useRef, useState } from "react";
+import * as d3 from "d3";
+
+type SortDirection = "asc" | "desc" | null;
+
+/**
+ * Interface that describes a linear data provider.
+ */
+export interface IDataProvider {
+    data: any[];
+}
+
+/**
+ * Interface that defines a more flexible variant for how to handle a column data.
+ */
+export interface ColumnDef<T = any> {
+    key: string;
+    label?: string;
+
+    /**
+     * Delegate to produce custom rendering & style of a specific value in a row.
+     * @param value The value to render.
+     * @param row The strongly typed row data.
+     * @param cell The D3 cell to update with custom rendering & styling.
+     */
+    render?: (value: any, row: T, cell: d3.Selection<HTMLTableCellElement, unknown, null, undefined>) => void;
+}
+
+/**
+ * Interface that describes a column provider.
+ */
+export interface IColumnProvider<T = any> {
+    columns: (string | ColumnDef<T>)[];
+}
+
+/**
+ * Generic CSV Table component.
+ * @returns React Component that represents a CSV table.
+ */
+export const CSVTable: React.FC<{ dataProvider: IDataProvider, columnProvider: IColumnProvider }> = observer(({ dataProvider, columnProvider }) => {
+
+    // #region -- Refs & State --
+
+    const containerRef = useRef<HTMLDivElement>(null);
+    const topScrollRef = useRef<HTMLDivElement>(null);
+    const topScrollInnerRef = useRef<HTMLDivElement>(null);
+    const [scrollWidth, setScrollWidth] = useState<number>(0);
+    const [sortColumn, setSortColumn] = useState<string | null>(null);
+    const [sortDirection, setSortDirection] = useState<SortDirection>(null);
+    const { modeColors, _hordeClasses, _detailClasses } = getHordeStyling();
+
+    const colorRecords: Record<string, string> = getColorRecords();
+    const hasConstructed = useRef<boolean>(false);
+
+    const scaffold = useRef<{
+        wrapper: d3.Selection<HTMLDivElement, unknown, null, undefined>;
+        table: d3.Selection<HTMLTableElement, unknown, null, undefined>;
+        thead: d3.Selection<HTMLTableSectionElement, unknown, null, undefined>;
+        theaderrow: d3.Selection<HTMLTableRowElement, unknown, null, undefined>;
+        tbody: d3.Selection<HTMLTableSectionElement, unknown, null, undefined>;
+        tooltip: d3.Selection<HTMLDivElement, unknown, HTMLElement, any>;
+    } | null>(null);
+
+    //#endregion -- Refs & State --
+
+    // #region -- Common Colors --
+
+    const summaryColorScaleBase = d3.scaleLinear<string>()
+        .domain([0, 0.5, 1])
+        .range([colorRecords["Failure"], colorRecords["Warnings"], colorRecords["Success"]]);
+
+    const summaryColorScale = (value: number): string => {
+        if (value === -1) {
+            return colorRecords["Waiting"];
+        }
+        return summaryColorScaleBase(value);
+    };
+
+    // #endregion -- Common Colors --
+
+    // #region -- Callback Construction -- 
+
+    const handleHeaderClick = useCallback((columnKey: string) => {
+        // If we have clicked on existing table, cycle the options.
+        // asc -> desc -> no sort -> asc -> ...
+        if (sortColumn === columnKey) {
+            if (sortDirection === "asc") {
+                setSortDirection("desc");
+            } else if (sortDirection === "desc") {
+                setSortColumn(null);
+                setSortDirection(null);
+            } else {
+                setSortDirection("asc");
+            }
+        } else {
+            setSortColumn(columnKey);
+            setSortDirection("asc");
+        }
+    }, [sortColumn, sortDirection]);
+
+    // We must stash the refs here in order to properly enclose the sortColumn useState
+    const handleHeaderClickRef = useRef(handleHeaderClick);
+    handleHeaderClickRef.current = handleHeaderClick;
+
+    const getSortIndicator = useCallback((columnKey: string): string => {
+        if (sortColumn !== columnKey) {
+            return "";
+        }
+        return sortDirection === "asc" ? " ↑" : " ↓";
+    }, [sortColumn, sortDirection]);
+
+    const sortedData = useMemo(() => {
+        if (!sortColumn || !sortDirection) {
+            return dataProvider.data;
+        }
+
+        return [...dataProvider.data].sort((a, b) => {
+            const aVal = a[sortColumn];
+            const bVal = b[sortColumn];
+
+            if (aVal === bVal) {
+                return 0;
+            }
+            if (aVal === null || aVal === undefined) {
+                return 1;
+            }
+            if (bVal === null || bVal === undefined) {
+                return -1;
+            }
+
+            let comparison: number;
+            if (typeof aVal === "number" && typeof bVal === "number") {
+                comparison = aVal - bVal;
+            } else if (typeof aVal === "string" && typeof bVal === "string") {
+                comparison = aVal.localeCompare(bVal);
+            } else {
+                comparison = String(aVal).localeCompare(String(bVal));
+            }
+
+            return sortDirection === "asc" ? comparison : -comparison;
+        });
+    }, [dataProvider.data, sortColumn, sortDirection]);
+
+    // Synchronize scrolling between top and bottom scrollbars
+    const handleTopScroll = useCallback(() => {
+        if (topScrollRef.current && containerRef.current) {
+            containerRef.current.scrollLeft = topScrollRef.current.scrollLeft;
+        }
+    }, []);
+
+    const handleBottomScroll = useCallback(() => {
+        if (topScrollRef.current && containerRef.current) {
+            topScrollRef.current.scrollLeft = containerRef.current.scrollLeft;
+        }
+    }, []);
+
+    // #endregion -- Callback Construction --
+
+    // #region -- Use Effects --
+
+    useEffect(() => {
+        if (!containerRef.current || scaffold.current) return;
+
+        // Create our base container.
+        const wrapper = d3.select(containerRef.current)
+            .append("div")
+            .attr("id", "csv-container");
+
+        // Create the table element.
+        const table = wrapper.append("table")
+            .style("border-collapse", "separate")
+            .style("border-spacing", "2px")
+            .style("width", "fit-content")
+            .style("font-family", preloadFonts[0])
+            .style("font-size", "12px")
+            .style("color", modeColors.text);
+
+        // Add header && tbody.
+        const thead = table.append("thead");
+        const tbody = table.append("tbody");
+
+        // Add tool tip div; will be used throughout.
+        d3.selectAll(".d3-tooltip").remove();
+
+        const tooltip = d3.select("body")
+            .append("div")
+            .attr("class", "d3-tooltip")
+            .style("position", "absolute")
+            .style("pointer-events", "none")
+            .style("padding", "6px 12px")
+            .style("background", "rgba(0,0,0,0.95)")
+            .style("color", "#fff")
+            .style("z-index", "5")
+            .style("min-width", "250px")
+            .style("font-family", preloadFonts[0])
+            .style("font-size", "12px")
+            .style("border-radius", "2px")
+            .style("line-height", "1.4")
+            .style("opacity", 0);
+
+        const theaderrow = thead.append("tr")
+            .style("height", "20px");
+
+        scaffold.current = { wrapper, table, thead, theaderrow, tbody, tooltip };
+
+        return () => {
+            d3.selectAll(".d3-tooltip").remove();
+        };
+    }, []);
+
+    type ColumnType = string | ColumnDef<any>;
+
+    useEffect(() => {
+        if (hasConstructed.current === false) {
+            if (!containerRef.current || !scaffold.current) return;
+
+            scaffold.current.theaderrow.selectAll<HTMLTableCellElement, ColumnType>("th.col")
+                .data(columnProvider.columns)
+                .join("th")
+                .attr("class", "col")
+                .style("position", "sticky")
+                .style("top", "0")
+                .style("z-index", "3")
+                .style("background-color", modeColors.content)
+                .style("border", `1px solid ${modeColors.background}`)
+                .style("padding", "8px")
+                .style("box-sizing", "border-box")
+                .style("fontWeight", "bold")
+                .style("min-width", "150px")
+                .style("cursor", "pointer")
+                .style("user-select", "none")
+                .on("click", function (_event, d) {
+                    const col = d as unknown as ColumnType;
+
+                    if (!col) {
+                        return;
+                    }
+
+                    const key = typeof col === "string" ? col : col.key;
+                    handleHeaderClickRef.current(key);
+                })
+                .text(d => {
+                    const label = typeof d === "string" ? d : (d.label ?? d.key);
+                    const key = typeof d === "string" ? d : d.key;
+                    return label + getSortIndicator(key);
+                });
+
+            hasConstructed.current = true;
+        }
+    }, [columnProvider.columns, getSortIndicator]);
+
+    const renderCell = (cell: d3.Selection<HTMLTableCellElement, unknown, null, undefined>, col: string |
+        ColumnDef<any>, rowData: any) => {
+        const key = typeof col === "string" ? col : col.key;
+
+        // Note: we are explicitly declaring any | undefined because an undefined value likely represents a computed column based on the row datum.
+        const value: any | undefined = rowData[key];
+
+        if (typeof col !== "string" && col.render) {
+            col.render(value, rowData, cell);
+        } else {
+            cell.text(value);
+        }
+    };
+
+    // Update header sort indicators when sort state changes
+    useEffect(() => {
+        if (!scaffold.current) {
+            return;
+        }
+
+        scaffold.current.theaderrow.selectAll("th.col")
+            .data(columnProvider.columns)
+            .text(d => {
+                const label = typeof d === "string" ? d : (d.label ?? d.key);
+                const key = typeof d === "string" ? d : d.key;
+                return label + getSortIndicator(key);
+            });
+    }, [sortColumn, sortDirection, columnProvider.columns, getSortIndicator]);
+
+    useEffect(() => {
+        if (!containerRef.current || !scaffold.current) return;
+        const rows = scaffold.current.tbody.selectAll("tr")
+            .data(sortedData)
+            .join("tr");
+
+        rows.each(function (rowData: any) {
+            const tr = d3.select(this);
+
+            tr.selectAll("td")
+                .data(columnProvider.columns)
+                .join(
+                    enter => enter.append("td")
+                        .style("border-right", `1px solid ${modeColors.content}`)
+                        .style("padding", "8px")
+                        .style("box-sizing", "border-box")
+                        .each(function (col) {
+                            renderCell(d3.select(this as HTMLTableCellElement), col, rowData);
+                        }),
+                    update => update.each(function (col) {
+                        const cell = d3.select(this as HTMLTableCellElement);
+                        cell.selectAll("*").remove();
+                        cell.text("");
+                        renderCell(cell, col, rowData);
+                    }),
+                    exit => exit.remove()
+                );
+        });
+
+    }, [sortedData])
+
+    // Update the top scrollbar width to match the table's scroll width
+    useEffect(() => {
+        if (!containerRef.current) return;
+
+        const updateScrollWidth = () => {
+            if (containerRef.current) {
+                setScrollWidth(containerRef.current.scrollWidth);
+            }
+        };
+
+        // Initial update
+        updateScrollWidth();
+
+        // Use ResizeObserver to detect when the table content changes size
+        const resizeObserver = new ResizeObserver(updateScrollWidth);
+        resizeObserver.observe(containerRef.current);
+
+        return () => {
+            resizeObserver.disconnect();
+        };
+    }, [dataProvider.data, columnProvider.columns]);
+
+    // #endregion -- Use Effects --
+
+    return (
+        <>
+            {/* Top scrollbar */}
+            <div
+                ref={topScrollRef}
+                onScroll={handleTopScroll}
+                style={{
+                    width: 'calc(100% - 4px)',
+                    overflowX: 'auto',
+                    overflowY: 'hidden',
+                    marginRight: '4px',
+                }}
+            >
+                <div
+                    ref={topScrollInnerRef}
+                    style={{
+                        width: scrollWidth > 0 ? `${scrollWidth}px` : '100%',
+                        height: '1px',
+                    }}
+                />
+            </div>
+
+            {/* Table container */}
+            <div
+                id="d3-container"
+                ref={containerRef}
+                onScroll={handleBottomScroll}
+                style={{
+                    height: '80%',
+                    width: 'calc(100% - 4px)',
+                    maxHeight: '80vh',
+                    overflowX: 'auto',
+                    overflowY: 'auto',
+                    display: 'block',
+                    marginRight: '4px',
+                }}
+            />
+        </>
+    )
+})
\ No newline at end of file
diff --git a/Engine/Source/Programs/Horde/HordeDashboard/plugins/structuredanalytics/api.ts b/Engine/Source/Programs/Horde/HordeDashboard/plugins/structuredanalytics/api.ts
new file mode 100644
--- /dev/null
+++ b/Engine/Source/Programs/Horde/HordeDashboard/plugins/structuredanalytics/api.ts
@@ -0,0 +1,61 @@
+import backend from "horde/backend";
+import { RequestCancellationManager } from "horde/backend/RequestCancellationManager";
+
+let epicSandboxCancellationManager: RequestCancellationManager = new RequestCancellationManager();
+
+// #region -- Analytics Response Types --
+
+export type JobSummary = {
+    streamId: string,
+    jobId: string,
+    templateId: string
+    passRatio: number,
+    passWithWarningRatio: number,
+    failureRatio: number,
+    warningRatio: number,
+    stepWarningCount: number,
+    stepFailureCount: number,
+    stepPassCount: number,
+    jobWallTime: number,
+    jobStepTotalTime: number
+}
+
+// #endregion -- Analytics Response Types --
+
+// #region -- Analytics Request APIs --
+
+/**
+ * Obtains the most recent job summaries.
+ * @returns A list of job summaries.
+ * @remark This is a temporary API endpoint that is working with relatively raw data.
+ */
+export async function getJobSummaries(): Promise<JobSummary[]> {
+    return new Promise<JobSummary[]>((resolve, reject) => {
+        backend.fetch.get(`/api/v1/buildhealth/analytics/jobsummary`).then((response) => {
+            let data = response.data as any;
+            const results: JobSummary[] = data.map((x) => ({
+                streamId: x["streamId"],
+                jobId: x["jobId"],
+                templateId: x["templateId"],
+                commitId: x["commitId"],
+                codeCommitId: x["codeCommitId"],
+                createTimeUtc: x["createTimeUtc"],
+                finishTimeUtc: x["finishTimeUtc"],
+                jobWallTime: x["jobWallTime"],
+                jobStepTotalTime: x["jobStepsTotalTime"],
+                passRatio: x["passRatio"],
+                passWithWarningRatio: x["passWithWarningRatio"],
+                failureRatio: x["failureRatio"],
+                warningRatio: x["warningRatio"],
+                stepWarningCount: x["stepWarningCount"],
+                stepFailureCount: x["stepFailureCount"],
+                stepPassCount: x["stepPassCount"],
+            }));
+            resolve(results);
+        }).catch(reason => {
+            reject(reason);
+        });
+    })
+}
+
+// #endregion -- Analytics Request APIs --
diff --git a/Engine/Source/Programs/Horde/HordeDashboard/plugins/structuredanalytics/index.tsx b/Engine/Source/Programs/Horde/HordeDashboard/plugins/structuredanalytics/index.tsx
new file mode 100644
--- /dev/null
+++ b/Engine/Source/Programs/Horde/HordeDashboard/plugins/structuredanalytics/index.tsx
@@ -0,0 +1,18 @@
+// Copyright Epic Games, Inc. All Rights Reserved.
+
+import { MountType, registerHordePlugin } from "hordePlugins";
+import { AnalyticsView } from "./AnalyticsView";
+registerHordePlugin({
+    id: "structuredanalytics",
+    routes: [
+        { path: "structuredanalytics", element: <AnalyticsView /> },
+    ],
+    mounts: [
+        {
+            type: MountType.TopNav,
+            context: "Tools",
+            text: "Analytics - Structured Analytics",
+            route: "/structuredanalytics"
+        }
+    ]
+});
\ No newline at end of file
