[NuGet] Update Microsoft.Extensions.* to 10.0.5

#rnx
#Synced-CL 52554515

diff --git a/Engine/Source/Programs/Horde/Drivers/JobDriver.Tests/JobDriver.Tests.csproj b/Engine/Source/Programs/Horde/Drivers/JobDriver.Tests/JobDriver.Tests.csproj
--- a/Engine/Source/Programs/Horde/Drivers/JobDriver.Tests/JobDriver.Tests.csproj
+++ b/Engine/Source/Programs/Horde/Drivers/JobDriver.Tests/JobDriver.Tests.csproj
@@ -1,7 +1,7 @@
 <Project Sdk="Microsoft.NET.Sdk">
 
   <ItemGroup>
-    <PackageReference Include="Microsoft.Extensions.Logging.Abstractions" Version="8.0.3" />
+    <PackageReference Include="Microsoft.Extensions.Logging.Abstractions" Version="10.0.5" />
     <PackageReference Include="Moq" Version="4.20.72" />
   </ItemGroup>
 
diff --git a/Engine/Source/Programs/Horde/Drivers/JobDriver/JobDriver.csproj b/Engine/Source/Programs/Horde/Drivers/JobDriver/JobDriver.csproj
--- a/Engine/Source/Programs/Horde/Drivers/JobDriver/JobDriver.csproj
+++ b/Engine/Source/Programs/Horde/Drivers/JobDriver/JobDriver.csproj
@@ -46,14 +46,14 @@
       <PrivateAssets>all</PrivateAssets>
       <IncludeAssets>runtime; build; native; contentfiles; analyzers; buildtransitive</IncludeAssets>
     </PackageReference>
-    <PackageReference Include="Microsoft.Extensions.Caching.Memory" Version="8.0.1" />
-    <PackageReference Include="Microsoft.Extensions.Hosting" Version="8.0.1" />
-    <PackageReference Include="Microsoft.Extensions.Hosting.Abstractions" Version="8.0.1" />
-    <PackageReference Include="Microsoft.Extensions.Hosting.WindowsServices" Version="8.0.1" />
-    <PackageReference Include="Microsoft.Extensions.Http" Version="8.0.1" />
-    <PackageReference Include="Microsoft.Extensions.Http.Polly" Version="8.0.20" />
-    <PackageReference Include="Microsoft.Extensions.Logging.Abstractions" Version="8.0.3" />
-    <PackageReference Include="Microsoft.Extensions.Options.DataAnnotations" Version="8.0.0" />
+    <PackageReference Include="Microsoft.Extensions.Caching.Memory" Version="10.0.5" />
+    <PackageReference Include="Microsoft.Extensions.Hosting" Version="10.0.5" />
+    <PackageReference Include="Microsoft.Extensions.Hosting.Abstractions" Version="10.0.5" />
+    <PackageReference Include="Microsoft.Extensions.Hosting.WindowsServices" Version="10.0.5" />
+    <PackageReference Include="Microsoft.Extensions.Http" Version="10.0.5" />
+    <PackageReference Include="Microsoft.Extensions.Http.Polly" Version="10.0.5" />
+    <PackageReference Include="Microsoft.Extensions.Logging.Abstractions" Version="10.0.5" />
+    <PackageReference Include="Microsoft.Extensions.Options.DataAnnotations" Version="10.0.5" />
     <PackageReference Include="Microsoft.VisualStudio.Threading.Analyzers" Version="17.14.15">
       <PrivateAssets>all</PrivateAssets>
       <IncludeAssets>runtime; build; native; contentfiles; analyzers; buildtransitive</IncludeAssets>
diff --git a/Engine/Source/Programs/Horde/HordeAgent.Tests/HordeAgent.Tests.csproj b/Engine/Source/Programs/Horde/HordeAgent.Tests/HordeAgent.Tests.csproj
--- a/Engine/Source/Programs/Horde/HordeAgent.Tests/HordeAgent.Tests.csproj
+++ b/Engine/Source/Programs/Horde/HordeAgent.Tests/HordeAgent.Tests.csproj
@@ -1,7 +1,7 @@
 <Project Sdk="Microsoft.NET.Sdk">
 
   <ItemGroup>
-    <PackageReference Include="Microsoft.Extensions.Logging.Abstractions" Version="8.0.3" />
+    <PackageReference Include="Microsoft.Extensions.Logging.Abstractions" Version="10.0.5" />
     <PackageReference Include="Moq" Version="4.20.72" />
   </ItemGroup>
 
diff --git a/Engine/Source/Programs/Horde/HordeAgent/HordeAgent.csproj b/Engine/Source/Programs/Horde/HordeAgent/HordeAgent.csproj
--- a/Engine/Source/Programs/Horde/HordeAgent/HordeAgent.csproj
+++ b/Engine/Source/Programs/Horde/HordeAgent/HordeAgent.csproj
@@ -34,14 +34,14 @@
       <PrivateAssets>all</PrivateAssets>
       <IncludeAssets>runtime; build; native; contentfiles; analyzers; buildtransitive</IncludeAssets>
     </PackageReference>
-    <PackageReference Include="Microsoft.Extensions.Caching.Memory" Version="8.0.1" />
-    <PackageReference Include="Microsoft.Extensions.Hosting" Version="8.0.1" />
-    <PackageReference Include="Microsoft.Extensions.Hosting.Abstractions" Version="8.0.1" />
-    <PackageReference Include="Microsoft.Extensions.Hosting.WindowsServices" Version="8.0.1" />
-    <PackageReference Include="Microsoft.Extensions.Http" Version="8.0.1" />
-    <PackageReference Include="Microsoft.Extensions.Http.Polly" Version="8.0.20" />
-    <PackageReference Include="Microsoft.Extensions.Logging.Abstractions" Version="8.0.3" />
-    <PackageReference Include="Microsoft.Extensions.Options.DataAnnotations" Version="8.0.0" />
+    <PackageReference Include="Microsoft.Extensions.Caching.Memory" Version="10.0.5" />
+    <PackageReference Include="Microsoft.Extensions.Hosting" Version="10.0.5" />
+    <PackageReference Include="Microsoft.Extensions.Hosting.Abstractions" Version="10.0.5" />
+    <PackageReference Include="Microsoft.Extensions.Hosting.WindowsServices" Version="10.0.5" />
+    <PackageReference Include="Microsoft.Extensions.Http" Version="10.0.5" />
+    <PackageReference Include="Microsoft.Extensions.Http.Polly" Version="10.0.5" />
+    <PackageReference Include="Microsoft.Extensions.Logging.Abstractions" Version="10.0.5" />
+    <PackageReference Include="Microsoft.Extensions.Options.DataAnnotations" Version="10.0.5" />
     <PackageReference Include="Microsoft.VisualStudio.Threading.Analyzers" Version="17.14.15">
       <PrivateAssets>all</PrivateAssets>
       <IncludeAssets>runtime; build; native; contentfiles; analyzers; buildtransitive</IncludeAssets>
diff --git a/Engine/Source/Programs/Horde/HordeCmd/HordeCmd.csproj b/Engine/Source/Programs/Horde/HordeCmd/HordeCmd.csproj
--- a/Engine/Source/Programs/Horde/HordeCmd/HordeCmd.csproj
+++ b/Engine/Source/Programs/Horde/HordeCmd/HordeCmd.csproj
@@ -14,7 +14,7 @@
   </ItemGroup>
 
   <ItemGroup>
-    <PackageReference Include="Microsoft.Extensions.Configuration.EnvironmentVariables" Version="8.0.0" />
+    <PackageReference Include="Microsoft.Extensions.Configuration.EnvironmentVariables" Version="10.0.5" />
     <PackageReference Include="Serilog" Version="3.1.1" />
     <PackageReference Include="Serilog.Extensions.Hosting" Version="8.0.0" />
     <PackageReference Include="Serilog.Extensions.Logging" Version="8.0.0" />
diff --git a/Engine/Source/Programs/Horde/HordeServer/HordeServer.csproj b/Engine/Source/Programs/Horde/HordeServer/HordeServer.csproj
--- a/Engine/Source/Programs/Horde/HordeServer/HordeServer.csproj
+++ b/Engine/Source/Programs/Horde/HordeServer/HordeServer.csproj
@@ -35,8 +35,8 @@
     <PackageReference Include="Microsoft.AspNetCore.DataProtection.StackExchangeRedis" Version="8.0.20" />
     <PackageReference Include="Microsoft.AspNetCore.Mvc.Razor.RuntimeCompilation" Version="8.0.20" />
     <PackageReference Include="Microsoft.AspNetCore.SpaServices.Extensions" Version="8.0.20" />
-    <PackageReference Include="Microsoft.Extensions.Hosting.WindowsServices" Version="8.0.1" />
-    <PackageReference Include="Microsoft.Extensions.Http.Polly" Version="8.0.20" />
+    <PackageReference Include="Microsoft.Extensions.Hosting.WindowsServices" Version="10.0.5" />
+    <PackageReference Include="Microsoft.Extensions.Http.Polly" Version="10.0.5" />
     <PackageReference Include="Microsoft.VisualStudio.Threading.Analyzers" Version="17.14.15">
       <PrivateAssets>all</PrivateAssets>
       <IncludeAssets>runtime; build; native; contentfiles; analyzers; buildtransitive</IncludeAssets>
diff --git a/Engine/Source/Programs/Horde/Tools/ReleaseNotes/ReleaseNotes.csproj b/Engine/Source/Programs/Horde/Tools/ReleaseNotes/ReleaseNotes.csproj
--- a/Engine/Source/Programs/Horde/Tools/ReleaseNotes/ReleaseNotes.csproj
+++ b/Engine/Source/Programs/Horde/Tools/ReleaseNotes/ReleaseNotes.csproj
@@ -7,7 +7,7 @@
   </PropertyGroup>
 
   <ItemGroup>
-    <PackageReference Include="Microsoft.Extensions.Logging.Console" Version="8.0.1" />
+    <PackageReference Include="Microsoft.Extensions.Logging.Console" Version="10.0.5" />
   </ItemGroup>
 
   <ItemGroup>
diff --git a/Engine/Source/Programs/OidcToken/OidcToken.csproj b/Engine/Source/Programs/OidcToken/OidcToken.csproj
--- a/Engine/Source/Programs/OidcToken/OidcToken.csproj
+++ b/Engine/Source/Programs/OidcToken/OidcToken.csproj
@@ -15,13 +15,13 @@
   </ItemGroup>
 
   <ItemGroup>
-    <PackageReference Include="Microsoft.Extensions.Configuration" Version="8.0.0" />
-    <PackageReference Include="Microsoft.Extensions.Configuration.Binder" Version="8.0.2" />
-    <PackageReference Include="Microsoft.Extensions.DependencyInjection" Version="8.0.1" />
-    <PackageReference Include="Microsoft.Extensions.Hosting" Version="8.0.1" />
-    <PackageReference Include="Microsoft.Extensions.Logging" Version="8.0.1" />
-    <PackageReference Include="Microsoft.Extensions.Logging.Console" Version="8.0.1" />
-    <PackageReference Include="Microsoft.Extensions.Options.DataAnnotations" Version="8.0.0" />
+    <PackageReference Include="Microsoft.Extensions.Configuration" Version="10.0.5" />
+    <PackageReference Include="Microsoft.Extensions.Configuration.Binder" Version="10.0.5" />
+    <PackageReference Include="Microsoft.Extensions.DependencyInjection" Version="10.0.5" />
+    <PackageReference Include="Microsoft.Extensions.Hosting" Version="10.0.5" />
+    <PackageReference Include="Microsoft.Extensions.Logging" Version="10.0.5" />
+    <PackageReference Include="Microsoft.Extensions.Logging.Console" Version="10.0.5" />
+    <PackageReference Include="Microsoft.Extensions.Options.DataAnnotations" Version="10.0.5" />
     <PackageReference Include="Serilog.Extensions.Hosting" Version="8.0.0" />
     <PackageReference Include="Serilog.Settings.Configuration" Version="8.0.4" />
     <PackageReference Include="Serilog.Sinks.Console" Version="6.0.0" />
diff --git a/Engine/Source/Programs/Shared/EpicGames.Build/EpicGames.Build.csproj b/Engine/Source/Programs/Shared/EpicGames.Build/EpicGames.Build.csproj
--- a/Engine/Source/Programs/Shared/EpicGames.Build/EpicGames.Build.csproj
+++ b/Engine/Source/Programs/Shared/EpicGames.Build/EpicGames.Build.csproj
@@ -15,6 +15,6 @@
   </ItemGroup>
 
   <ItemGroup>
-    <PackageReference Include="Microsoft.Extensions.FileSystemGlobbing" Version="8.0.0" />
+    <PackageReference Include="Microsoft.Extensions.FileSystemGlobbing" Version="10.0.5" />
   </ItemGroup>
 </Project>
diff --git a/Engine/Source/Programs/Shared/EpicGames.Core/EpicGames.Core.csproj b/Engine/Source/Programs/Shared/EpicGames.Core/EpicGames.Core.csproj
--- a/Engine/Source/Programs/Shared/EpicGames.Core/EpicGames.Core.csproj
+++ b/Engine/Source/Programs/Shared/EpicGames.Core/EpicGames.Core.csproj
@@ -5,10 +5,10 @@
 
   <ItemGroup>
     <PackageReference Include="JetBrains.Annotations" Version="2024.3.0" />
-    <PackageReference Include="Microsoft.Extensions.Configuration" Version="8.0.0" />
-    <PackageReference Include="Microsoft.Extensions.Http.Polly" Version="8.0.20" />
-    <PackageReference Include="Microsoft.Extensions.Logging.Console" Version="8.0.1" />
-    <PackageReference Include="Microsoft.Extensions.ObjectPool" Version="8.0.20" />
+    <PackageReference Include="Microsoft.Extensions.Configuration" Version="10.0.5" />
+    <PackageReference Include="Microsoft.Extensions.Http.Polly" Version="10.0.5" />
+    <PackageReference Include="Microsoft.Extensions.Logging.Console" Version="10.0.5" />
+    <PackageReference Include="Microsoft.Extensions.ObjectPool" Version="10.0.5" />
     <PackageReference Include="OpenTracing" Version="0.12.1" />
   </ItemGroup>
 </Project>
diff --git a/Engine/Source/Programs/Shared/EpicGames.Horde/EpicGames.Horde.csproj b/Engine/Source/Programs/Shared/EpicGames.Horde/EpicGames.Horde.csproj
--- a/Engine/Source/Programs/Shared/EpicGames.Horde/EpicGames.Horde.csproj
+++ b/Engine/Source/Programs/Shared/EpicGames.Horde/EpicGames.Horde.csproj
@@ -14,8 +14,8 @@
       <IncludeAssets>runtime; build; native; contentfiles; analyzers; buildtransitive</IncludeAssets>
     </PackageReference>
     <PackageReference Include="K4os.Compression.LZ4" Version="1.2.16" />
-    <PackageReference Include="Microsoft.Data.Sqlite" Version="8.0.20" />
-    <PackageReference Include="Microsoft.Extensions.Caching.Memory" Version="8.0.1" />
+    <PackageReference Include="Microsoft.Data.Sqlite" Version="10.0.5" />
+    <PackageReference Include="Microsoft.Extensions.Caching.Memory" Version="10.0.5" />
     <PackageReference Include="Microsoft.VisualStudio.Threading.Analyzers" Version="17.14.15">
       <PrivateAssets>all</PrivateAssets>
       <IncludeAssets>runtime; build; native; contentfiles; analyzers; buildtransitive</IncludeAssets>
diff --git a/Engine/Source/Programs/Shared/EpicGames.LogMatchers/EpicGames.LogMatchers.csproj b/Engine/Source/Programs/Shared/EpicGames.LogMatchers/EpicGames.LogMatchers.csproj
--- a/Engine/Source/Programs/Shared/EpicGames.LogMatchers/EpicGames.LogMatchers.csproj
+++ b/Engine/Source/Programs/Shared/EpicGames.LogMatchers/EpicGames.LogMatchers.csproj
@@ -1,6 +1,6 @@
 ﻿<Project Sdk="Microsoft.NET.Sdk">
   <ItemGroup>
-    <PackageReference Include="Microsoft.Extensions.Logging" Version="8.0.1" />
+    <PackageReference Include="Microsoft.Extensions.Logging" Version="10.0.5" />
   </ItemGroup>
 
   <ItemGroup>
diff --git a/Engine/Source/Programs/Shared/EpicGames.OIDC/EpicGames.OIDC.csproj b/Engine/Source/Programs/Shared/EpicGames.OIDC/EpicGames.OIDC.csproj
--- a/Engine/Source/Programs/Shared/EpicGames.OIDC/EpicGames.OIDC.csproj
+++ b/Engine/Source/Programs/Shared/EpicGames.OIDC/EpicGames.OIDC.csproj
@@ -5,9 +5,9 @@
 
   <ItemGroup>
     <PackageReference Include="IdentityModel.OidcClient" Version="6.0.0" />
-    <PackageReference Include="Microsoft.Extensions.Configuration.Binder" Version="8.0.2" />
-    <PackageReference Include="Microsoft.Extensions.Configuration.Json" Version="8.0.1" />
-    <PackageReference Include="Microsoft.Extensions.Options" Version="8.0.2" />
+    <PackageReference Include="Microsoft.Extensions.Configuration.Binder" Version="10.0.5" />
+    <PackageReference Include="Microsoft.Extensions.Configuration.Json" Version="10.0.5" />
+    <PackageReference Include="Microsoft.Extensions.Options" Version="10.0.5" />
   </ItemGroup>
 
   <ItemGroup>
diff --git a/Engine/Source/Programs/Shared/EpicGames.Perforce.Fixture/EpicGames.Perforce.Fixture.csproj b/Engine/Source/Programs/Shared/EpicGames.Perforce.Fixture/EpicGames.Perforce.Fixture.csproj
--- a/Engine/Source/Programs/Shared/EpicGames.Perforce.Fixture/EpicGames.Perforce.Fixture.csproj
+++ b/Engine/Source/Programs/Shared/EpicGames.Perforce.Fixture/EpicGames.Perforce.Fixture.csproj
@@ -10,7 +10,7 @@
   </ItemGroup>
 
   <ItemGroup>
-    <PackageReference Include="Microsoft.Extensions.Logging.Console" Version="8.0.1" />
+    <PackageReference Include="Microsoft.Extensions.Logging.Console" Version="10.0.5" />
   </ItemGroup>
 
   <ItemGroup>
diff --git a/Engine/Source/Programs/Shared/EpicGames.Slack/EpicGames.Slack.csproj b/Engine/Source/Programs/Shared/EpicGames.Slack/EpicGames.Slack.csproj
--- a/Engine/Source/Programs/Shared/EpicGames.Slack/EpicGames.Slack.csproj
+++ b/Engine/Source/Programs/Shared/EpicGames.Slack/EpicGames.Slack.csproj
@@ -5,8 +5,8 @@
   </PropertyGroup>
 
   <ItemGroup>
-    <PackageReference Include="Microsoft.Extensions.Http" Version="8.0.1" />
-    <PackageReference Include="Microsoft.Extensions.Http.Polly" Version="8.0.20" />
-    <PackageReference Include="Microsoft.Extensions.Logging.Abstractions" Version="8.0.3" />
+    <PackageReference Include="Microsoft.Extensions.Http" Version="10.0.5" />
+    <PackageReference Include="Microsoft.Extensions.Http.Polly" Version="10.0.5" />
+    <PackageReference Include="Microsoft.Extensions.Logging.Abstractions" Version="10.0.5" />
   </ItemGroup>
 </Project>
diff --git a/Engine/Source/Programs/Shared/EpicGames.UBA/EpicGames.UBA.csproj b/Engine/Source/Programs/Shared/EpicGames.UBA/EpicGames.UBA.csproj
--- a/Engine/Source/Programs/Shared/EpicGames.UBA/EpicGames.UBA.csproj
+++ b/Engine/Source/Programs/Shared/EpicGames.UBA/EpicGames.UBA.csproj
@@ -20,7 +20,7 @@
   <Import Project="Library.props" />
 
   <ItemGroup>
-    <PackageReference Include="Microsoft.Extensions.Logging" Version="8.0.1" />
+    <PackageReference Include="Microsoft.Extensions.Logging" Version="10.0.5" />
   </ItemGroup>
 
   <ItemGroup>
diff --git a/Engine/Source/Programs/UnrealBuildTool/UnrealBuildTool.csproj b/Engine/Source/Programs/UnrealBuildTool/UnrealBuildTool.csproj
--- a/Engine/Source/Programs/UnrealBuildTool/UnrealBuildTool.csproj
+++ b/Engine/Source/Programs/UnrealBuildTool/UnrealBuildTool.csproj
@@ -69,7 +69,7 @@
 
   <ItemGroup>
     <PackageReference Include="Microsoft.CodeAnalysis.CSharp" Version="4.13.0" />
-    <PackageReference Include="Microsoft.Extensions.Logging.Console" Version="8.0.1" />
+    <PackageReference Include="Microsoft.Extensions.Logging.Console" Version="10.0.5" />
     <PackageReference Include="Microsoft.VisualStudio.Setup.Configuration.Interop" Version="3.14.2075" />
     <PackageReference Include="Microsoft.VisualStudio.SolutionPersistence" Version="1.0.52" />
     <PackageReference Include="OpenTracing" Version="0.12.1" />
diff --git a/Engine/Source/Programs/UnrealGameSync/UnrealGameSyncShared/UnrealGameSyncShared.csproj b/Engine/Source/Programs/UnrealGameSync/UnrealGameSyncShared/UnrealGameSyncShared.csproj
--- a/Engine/Source/Programs/UnrealGameSync/UnrealGameSyncShared/UnrealGameSyncShared.csproj
+++ b/Engine/Source/Programs/UnrealGameSync/UnrealGameSyncShared/UnrealGameSyncShared.csproj
@@ -25,7 +25,7 @@
   </Choose>
   
   <ItemGroup>
-    <PackageReference Include="Microsoft.Extensions.Caching.Memory" Version="8.0.1" />
+    <PackageReference Include="Microsoft.Extensions.Caching.Memory" Version="10.0.5" />
   </ItemGroup>
 
   <ItemGroup>
