[Test Automation Hub] Remove heading and trailing spaces on user search input

#rnx
#Synced-CL 52573604

diff --git a/Engine/Source/Programs/Horde/HordeDashboard/plugins/test-automation/testAutomationPhaseGridView.tsx b/Engine/Source/Programs/Horde/HordeDashboard/plugins/test-automation/testAutomationPhaseGridView.tsx
--- a/Engine/Source/Programs/Horde/HordeDashboard/plugins/test-automation/testAutomationPhaseGridView.tsx
+++ b/Engine/Source/Programs/Horde/HordeDashboard/plugins/test-automation/testAutomationPhaseGridView.tsx
@@ -161,6 +161,7 @@ const PhaseGroupFilter: React.FC<{groups: PhaseStatusGroup[], selected?: string[
     }, [groups]);
     
     const filterSelectedTests = useCallback((filterText: string): ITag[] => {
+        filterText = filterText.trim();
         if (filterText.length < 2) return [];
         const lowerText = filterText.toLowerCase();
         if (filterText.startsWith("#")) {
diff --git a/Engine/Source/Programs/Horde/HordeDashboard/plugins/test-automation/testAutomationPhaseView.tsx b/Engine/Source/Programs/Horde/HordeDashboard/plugins/test-automation/testAutomationPhaseView.tsx
--- a/Engine/Source/Programs/Horde/HordeDashboard/plugins/test-automation/testAutomationPhaseView.tsx
+++ b/Engine/Source/Programs/Horde/HordeDashboard/plugins/test-automation/testAutomationPhaseView.tsx
@@ -250,6 +250,7 @@ const PhaseFilter: React.FC<{phases: TestPhaseStatus[], selected?: string[], onC
     }, [phases]);
     
     const filterSelectedTests = useCallback((filterText: string): ITag[] => {
+        filterText = filterText.trim();
         if (filterText.length < 2) return [];
         const lowerText = filterText.toLowerCase();
         if (filterText.startsWith("#")) {
diff --git a/Engine/Source/Programs/Horde/HordeDashboard/plugins/test-automation/testAutomationView.tsx b/Engine/Source/Programs/Horde/HordeDashboard/plugins/test-automation/testAutomationView.tsx
--- a/Engine/Source/Programs/Horde/HordeDashboard/plugins/test-automation/testAutomationView.tsx
+++ b/Engine/Source/Programs/Horde/HordeDashboard/plugins/test-automation/testAutomationView.tsx
@@ -81,6 +81,7 @@ const TestChooser: React.FC = observer(() => {
 
     const testNames: ITag[] = streamTests.map(t => ({key: `@${t.toLowerCase()}`, name: `@${t}`}));
     const filterSelectedTests = useCallback((filterText: string, _: ITag[]): ITag[] => {
+        filterText = filterText.trim();
         if (filterText.length < 2) return [];
         const lowerText = filterText.toLowerCase();
         const closestNames = testNames.values().filter((tag) => (tag.key as string).indexOf(lowerText) >= 0).take(10).toArray();
@@ -138,6 +139,7 @@ const TagChooser: React.FC = observer(() => {
 
     const testTags: ITag[] = streamTags.map(t => ({key: t.toLowerCase(), name: `#${t}`}));
     const filterSelectedTags = useCallback((filterText: string, _: ITag[]): ITag[] => {
+        filterText = filterText.trim();
         if (filterText.startsWith("#")) filterText = filterText.substring(1);
         if (filterText.length < 2) return [];
         const lowerText = filterText.toLowerCase();
