Horde: Include exception message in notification task error log

diff --git a/Engine/Source/Programs/Horde/Plugins/Build/HordeServer.Build/Notifications/NotificationService.cs b/Engine/Source/Programs/Horde/Plugins/Build/HordeServer.Build/Notifications/NotificationService.cs
--- a/Engine/Source/Programs/Horde/Plugins/Build/HordeServer.Build/Notifications/NotificationService.cs
+++ b/Engine/Source/Programs/Horde/Plugins/Build/HordeServer.Build/Notifications/NotificationService.cs
@@ -390,7 +390,7 @@ namespace HordeServer.Notifications
 				}
 				catch (Exception e)
 				{
-					_logger.LogError(e, "Exception while executing notification task");
+					_logger.LogError(e, "Exception while executing notification task: {Message}", e.Message);
 				}
 			});
 			
