Horde: Log agent relay errors when allocating compute resources

Previously the error was only recorded in OpenTelemetry traces and returned to the HTTP caller, with no log entry.

diff --git a/Engine/Source/Programs/Horde/Plugins/Compute/HordeServer.Compute/Compute/ComputeService.cs b/Engine/Source/Programs/Horde/Plugins/Compute/HordeServer.Compute/Compute/ComputeService.cs
--- a/Engine/Source/Programs/Horde/Plugins/Compute/HordeServer.Compute/Compute/ComputeService.cs
+++ b/Engine/Source/Programs/Horde/Plugins/Compute/HordeServer.Compute/Compute/ComputeService.cs
@@ -516,6 +516,7 @@ namespace HordeServer.Compute
 			}
 			catch (AgentRelayException are)
 			{
+				_logger.LogError(are, "Agent relay error while allocating compute resource for cluster {ClusterId}: {Message}", arp.ClusterId, are.Message);
 				span.SetStatus(Status.Error);
 				span.RecordException(are);
 				throw new ComputeServiceException(are.Message, are) { ShowToUser = true };
