Horde: Fix PerforceMaterializer exception when edge server switches with pending dirty files

The client name is persisted in the state file, which causes confusion when the edge server is changed.

diff --git a/Engine/Source/Programs/Horde/Drivers/JobDriver.Tests/PerforceMaterializerTests.cs b/Engine/Source/Programs/Horde/Drivers/JobDriver.Tests/PerforceMaterializerTests.cs
--- a/Engine/Source/Programs/Horde/Drivers/JobDriver.Tests/PerforceMaterializerTests.cs
+++ b/Engine/Source/Programs/Horde/Drivers/JobDriver.Tests/PerforceMaterializerTests.cs
@@ -287,6 +287,41 @@ public class PerforceMaterializerTests : BasePerforceFixtureTest
 		Assert.IsTrue(File.Exists(markerPath));
 	}
 
+	[TestMethod]
+	public async Task Sync_WithChangedClientNameAndDirtyFiles_ShouldRevertAndSync_Async()
+	{
+		CancellationToken ct = CancellationToken.None;
+		using PerforceMaterializer pm = CreatePm(PerforceSnapshot.FooMain);
+
+		// Sync to CL 2 and finalize so state is saved to disk
+		await pm.SyncAsync(2, -1, new SyncOptions { RemoveUntracked = false }, ct);
+		await pm.FinalizeAsync(ct);
+
+		// Write a marker file to later verify the workspace was not wiped
+		string markerPath = Path.Join(pm.SyncDir.FullName, "_edge_dirty_marker.txt");
+		await File.WriteAllTextAsync(markerPath, "marker", ct);
+
+		// Tamper with State.json to simulate an edge server switch with leftover DirtyFiles
+		PerforceMaterializer.State? state = await pm.LoadStateForTestAsync(ct);
+		Assert.IsNotNull(state);
+		string oldClientName = "Horde+PM+FAKE+id+fakeEdge";
+		string relativePath = PerforceSnapshot.GetSnapshot().GetStream(PerforceSnapshot.FooMain).GetWorkspace(2).Files.First().Path;
+		state = state with
+		{
+			Client = oldClientName,
+			Status = PerforceMaterializer.TransactionStatus.Clean,
+			DirtyFiles = [new PerforceMaterializer.WorkspaceFile($"//{oldClientName}/{relativePath}", Path.Join(pm.SyncDir.FullName, relativePath))]
+		};
+		string stateFilePath = Path.Join(pm.BaseDir.FullName, "State.json");
+		await File.WriteAllTextAsync(stateFilePath, JsonSerializer.Serialize(state, new JsonSerializerOptions { WriteIndented = true }), ct);
+
+		// Sync to CL 5 - should not crash; the fix remaps DirtyFiles to the current client
+		await pm.SyncAsync(5, -1, new SyncOptions { RemoveUntracked = false }, ct);
+
+		// Marker file should still exist (workspace was not wiped)
+		Assert.IsTrue(File.Exists(markerPath));
+	}
+
 	[TestMethod]
 	public async Task Sync_WithChangedClientNameAndStream_ShouldWipe_Async()
 	{
diff --git a/Engine/Source/Programs/Horde/Drivers/JobDriver/Execution/PerforceMaterializer.cs b/Engine/Source/Programs/Horde/Drivers/JobDriver/Execution/PerforceMaterializer.cs
--- a/Engine/Source/Programs/Horde/Drivers/JobDriver/Execution/PerforceMaterializer.cs
+++ b/Engine/Source/Programs/Horde/Drivers/JobDriver/Execution/PerforceMaterializer.cs
@@ -648,7 +648,20 @@ public sealed class PerforceMaterializer : IWorkspaceMaterializer
 				FileUtils.ForceDeleteFile(wf.LocalFile);
 			}
 			
-			List<string> dirtyClientFiles = state.DirtyFiles.Select(x => $"{x.ClientFile}@{state.Changelist}").ToList();
+			// When an agent is reassigned to a different edge server between jobs, DirtyFiles
+			// persisted in State.json still carry the old client name prefix (e.g. //buildedge_01/...).
+			// Remap them to the current client so the force-sync targets the right workspace.
+			string oldClientPrefix = $"//{state.Client}/";
+			string newClientPrefix = $"//{_perforceSettings!.ClientName}/";
+			List<string> dirtyClientFiles = state.DirtyFiles
+				.Select(x =>
+				{
+					string clientFile = x.ClientFile.StartsWith(oldClientPrefix, StringComparison.Ordinal)
+						? newClientPrefix + x.ClientFile[oldClientPrefix.Length..]
+						: x.ClientFile;
+					return $"{clientFile}@{state.Changelist}";
+				})
+				.ToList();
 			List<SyncRecord> syncedFiles = await _perforceWithClient.SyncAsync(
 				EpicGames.Perforce.SyncOptions.Force, -1, _options.SyncThreads, _options.SyncBatch, _options.SyncBatchSize,
 				-1, -1, dirtyClientFiles, cancellationToken).ToListAsync(cancellationToken);
