[Horde][PerformanceTrends][Server] Adding frametime stat to key stats.

#Synced-CL 52610330

diff --git a/Engine/Source/Programs/Horde/Plugins/PerformanceTrends/HordeServer.PerformanceTrends/KeyStatsPerformanceSummaryHandler.cs b/Engine/Source/Programs/Horde/Plugins/PerformanceTrends/HordeServer.PerformanceTrends/KeyStatsPerformanceSummaryHandler.cs
--- a/Engine/Source/Programs/Horde/Plugins/PerformanceTrends/HordeServer.PerformanceTrends/KeyStatsPerformanceSummaryHandler.cs
+++ b/Engine/Source/Programs/Horde/Plugins/PerformanceTrends/HordeServer.PerformanceTrends/KeyStatsPerformanceSummaryHandler.cs
@@ -60,6 +60,7 @@ namespace HordeServer.LicenseePromotionCandidate.PerformanceTrends
 									{EpicGames_Analytics_KeyStatTelemetryRecordGen.HitchTimePercent},
 									{EpicGames_Analytics_KeyStatTelemetryRecordGen.GameThreadtimeAvg},
 									{EpicGames_Analytics_KeyStatTelemetryRecordGen.RenderThreadtimeAvg},
+									{EpicGames_Analytics_KeyStatTelemetryRecordGen.FrametimeAvg},
 									{EpicGames_Analytics_KeyStatTelemetryRecordGen.SummaryName},
 									{EpicGames_Analytics_KeyStatTelemetryRecordGen.DynamicResolutionPercentageAvg},
 									{EpicGames_Analytics_KeyStatTelemetryRecordGen.EngineVersion},
diff --git a/Engine/Source/Programs/Shared/EpicGames.Analytics.PerformanceTrends/KeyStatTelemetryRecord.cs b/Engine/Source/Programs/Shared/EpicGames.Analytics.PerformanceTrends/KeyStatTelemetryRecord.cs
--- a/Engine/Source/Programs/Shared/EpicGames.Analytics.PerformanceTrends/KeyStatTelemetryRecord.cs
+++ b/Engine/Source/Programs/Shared/EpicGames.Analytics.PerformanceTrends/KeyStatTelemetryRecord.cs
@@ -74,7 +74,13 @@ namespace EpicGames.Analytics
 		public float RenderThreadtimeAvg { get; init; }
 
 		/// <summary>
-		/// Dynamic reoslution percentage average.
+		/// Frame time average.
+		/// </summary>
+		[Column("frametime_avg")]
+		public float FrametimeAvg { get; init; }
+
+		/// <summary>
+		/// Dynamic resolution percentage average.
 		/// </summary>
 		[Column("dynamic_resolution_percentage_avg")]
 		public float DynamicResolutionPercentageAvg { get; init; }
