[Horde][PerformanceTrends][Dashboard] Adding Frametime to KeyStats

#Synced-CL 52610330

diff --git a/Engine/Source/Programs/Horde/HordeDashboard/plugins/performancetrends/metrictypes/KeyStatsData.ts b/Engine/Source/Programs/Horde/HordeDashboard/plugins/performancetrends/metrictypes/KeyStatsData.ts
--- a/Engine/Source/Programs/Horde/HordeDashboard/plugins/performancetrends/metrictypes/KeyStatsData.ts
+++ b/Engine/Source/Programs/Horde/HordeDashboard/plugins/performancetrends/metrictypes/KeyStatsData.ts
@@ -24,6 +24,9 @@ export class KeyStatsData extends PerformanceTrendContext implements ISummaryMet
     @viewable({ label: 'Render Thread', category: 'Timing', unit: 'ms', defaultVisible: true, precision: 2 })
     renderThreadtimeAvg: number;
 
+    @viewable({ label: 'Frame Time', category: 'Timing', unit: 'ms', defaultVisible: true, precision: 2 })
+    frametimeAvg: number;
+
     @viewable({ label: 'Peak Physical Memory', category: 'Memory', unit: 'MB', defaultVisible: true, precision: 1 })
     physicalUsedMbMax: number;
 
diff --git a/Engine/Source/Programs/Horde/HordeDashboard/plugins/performancetrends/viewgenerators/KeyStatsViewGenerator.tsx b/Engine/Source/Programs/Horde/HordeDashboard/plugins/performancetrends/viewgenerators/KeyStatsViewGenerator.tsx
--- a/Engine/Source/Programs/Horde/HordeDashboard/plugins/performancetrends/viewgenerators/KeyStatsViewGenerator.tsx
+++ b/Engine/Source/Programs/Horde/HordeDashboard/plugins/performancetrends/viewgenerators/KeyStatsViewGenerator.tsx
@@ -157,6 +157,7 @@ export class KeyStatsViewGenerator implements IMetricViewGenerator<KeyStatsData>
         const memoryMeta = viewableProps.get(keyOf<KeyStatsData>('physicalUsedMbMax'))!;
         const gameThreadMeta = viewableProps.get(keyOf<KeyStatsData>('gameThreadtimeAvg'))!;
         const renderThreadMeta = viewableProps.get(keyOf<KeyStatsData>('renderThreadtimeAvg'))!;
+        const frametimeMeta = viewableProps.get(keyOf<KeyStatsData>('frametimeAvg'))!;
 
         // Create series arrays for each metric type (filtered by enabled properties)
         const gpuSeries = createGroupedSeries<KeyStatsData>(
@@ -194,6 +195,13 @@ export class KeyStatsViewGenerator implements IMetricViewGenerator<KeyStatsData>
             5
         );
 
+        const frametimeSeries = createGroupedSeries<KeyStatsData>(
+            keyOf<KeyStatsData>('frametimeAvg'),
+            (m) => m.frametimeAvg,
+            frametimeMeta.label,
+            8
+        );
+
         // Calculate averages across all platforms
         const gpuTimes: number[] = metrics.map(m => m.gpuTimeAvg);
         const memoryUsedMax: number[] = metrics.map(m => m.physicalUsedMbMax);
@@ -253,7 +261,7 @@ export class KeyStatsViewGenerator implements IMetricViewGenerator<KeyStatsData>
                         <span>Thread Times</span> <br />
                     </div>
                     <LineGraph
-                        series={[...gameThreadSeries, ...renderThreadSeries]}
+                        series={[...gameThreadSeries, ...renderThreadSeries, ...frametimeSeries]}
                         xAxisLabel="Changelist"
                         yAxisLabel={buildYAxisLabel(gameThreadMeta)}
                         width={width}
