Horde: Expose server environment on /api/v1/server/info

Primarily to help the dashboard identify what environment it's connected to.

diff --git a/Engine/Source/Programs/Horde/HordeServer/Server/ServerController.cs b/Engine/Source/Programs/Horde/HordeServer/Server/ServerController.cs
--- a/Engine/Source/Programs/Horde/HordeServer/Server/ServerController.cs
+++ b/Engine/Source/Programs/Horde/HordeServer/Server/ServerController.cs
@@ -30,16 +30,18 @@ namespace HordeServer.Server
 	{
 		readonly IAgentVersionProvider? _agentVersionProvider;
 		readonly IPluginCollection _pluginCollection;
+		readonly IServerInfo _serverInfo;
 		readonly IOptionsSnapshot<GlobalConfig> _globalConfig;
 		readonly ILogger<ServerController> _logger;
 
 		/// <summary>
 		/// Constructor
 		/// </summary>
-		public ServerController(IEnumerable<IAgentVersionProvider> agentVersionProviders, IPluginCollection pluginCollection, IOptionsSnapshot<GlobalConfig> globalConfig, ILogger<ServerController> logger)
+		public ServerController(IEnumerable<IAgentVersionProvider> agentVersionProviders, IPluginCollection pluginCollection, IServerInfo serverInfo, IOptionsSnapshot<GlobalConfig> globalConfig, ILogger<ServerController> logger)
 		{
 			_agentVersionProvider = agentVersionProviders.FirstOrDefault();
 			_pluginCollection = pluginCollection;
+			_serverInfo = serverInfo;
 			_globalConfig = globalConfig;
 			_logger = logger;
 		}
@@ -82,6 +84,8 @@ namespace HordeServer.Server
 				return new ServerPluginInfoResponse(plugin.Metadata.Name.ToString(), plugin.Metadata.Description, true, pluginVersion.ProductVersion ?? String.Empty, userAcls);
 			}).ToArray();
 
+			response.Environment = _serverInfo.Environment;
+
 			return response;
 		}
 
diff --git a/Engine/Source/Programs/Shared/EpicGames.Horde/Server/ServerMessages.cs b/Engine/Source/Programs/Shared/EpicGames.Horde/Server/ServerMessages.cs
--- a/Engine/Source/Programs/Shared/EpicGames.Horde/Server/ServerMessages.cs
+++ b/Engine/Source/Programs/Shared/EpicGames.Horde/Server/ServerMessages.cs
@@ -33,6 +33,10 @@ namespace EpicGames.Horde.Server
 		/// </summary>
 		public ServerPluginInfoResponse[] Plugins { get; set; } = [];
 
+		/// <summary>
+		/// The environment that the server is running in (e.g. "Development", "Production")
+		/// </summary>
+		public string? Environment { get; set; }
 	}
 
 	/// <summary>
